Start of OnlineServices Plugin implementation (OSSv2)

diff --git a/Engine/Plugins/Online/Experimental/OnlineServices/OnlineServices.uplugin b/Engine/Plugins/Online/Experimental/OnlineServices/OnlineServices.uplugin
new file mode 100644
--- /dev/null
+++ b/Engine/Plugins/Online/Experimental/OnlineServices/OnlineServices.uplugin
@@ -0,0 +1,25 @@
+{
+	"FileVersion": 3,
+
+	"FriendlyName": "Online Services",
+	"Version": 1,
+	"VersionName": "1.0",
+	"Description": "Shared code for interacting with online services implementations.",
+	"Category": "Online Platform",
+	"CreatedBy": "Epic Games, Inc.",
+	"CreatedByURL": "http://epicgames.com",
+	"EnabledByDefault": false,
+
+	"Modules": [
+		{
+			"Name": "OnlineServicesInterface",
+			"Type": "Runtime",
+			"LoadingPhase": "PostConfigInit"
+		},
+		{
+			"Name": "OnlineServicesCommon",
+			"Type": "Runtime",
+			"LoadingPhase": "PostConfigInit"
+		}
+	]
+}
\ No newline at end of file
diff --git a/Engine/Plugins/Online/Experimental/OnlineServices/Source/OnlineServicesCommon/OnlineServicesCommon.Build.cs b/Engine/Plugins/Online/Experimental/OnlineServices/Source/OnlineServicesCommon/OnlineServicesCommon.Build.cs
new file mode 100644
--- /dev/null
+++ b/Engine/Plugins/Online/Experimental/OnlineServices/Source/OnlineServicesCommon/OnlineServicesCommon.Build.cs
@@ -0,0 +1,29 @@
+// Copyright Epic Games, Inc. All Rights Reserved.
+
+using UnrealBuildTool;
+
+public class OnlineServicesCommon : ModuleRules
+{
+	public OnlineServicesCommon(ReadOnlyTargetRules Target) : base(Target)
+	{
+		PCHUsage = ModuleRules.PCHUsageMode.UseExplicitOrSharedPCHs;
+
+		bUseUnity = false; // set false to catch non-unity compile issues while prototyping
+
+		PublicDependencyModuleNames.AddRange(
+			new string[] {
+				"OnlineServicesInterface"
+			}
+		);
+
+		PublicIncludePaths.Add(ModuleDirectory);
+
+		// OnlineService cannot depend on Engine!
+		PrivateDependencyModuleNames.AddRange(
+			new string[] { 
+				"Core",
+				"CoreUObject"		// CoreUObject temporary dependency
+			}
+		);
+	}
+}
diff --git a/Engine/Plugins/Online/Experimental/OnlineServices/Source/OnlineServicesCommon/Private/Online/AuthCommon.cpp b/Engine/Plugins/Online/Experimental/OnlineServices/Source/OnlineServicesCommon/Private/Online/AuthCommon.cpp
new file mode 100644
--- /dev/null
+++ b/Engine/Plugins/Online/Experimental/OnlineServices/Source/OnlineServicesCommon/Private/Online/AuthCommon.cpp
@@ -0,0 +1,42 @@
+// Copyright Epic Games, Inc. All Rights Reserved.
+
+#include "Online/AuthCommon.h"
+
+#include "Online/OnlineAsyncOp.h"
+#include "Online/OnlineServicesCommon.h"
+
+namespace UE::Online {
+
+FAuthCommon::FAuthCommon(FOnlineServicesCommon& InServices)
+	: TOnlineComponent(TEXT("Auth"), InServices.AsShared())
+	, Services(InServices)
+{
+}
+
+TOnlineAsyncOpHandle<FAuthLogin> FAuthCommon::Login(FAuthLogin::Params&& Params)
+{
+	TOnlineAsyncOp<FAuthLogin>& Operation = Services.OpCache.GetOp<FAuthLogin>(MoveTemp(Params));
+	Operation.SetError(Errors::Unimplemented());
+	return Operation.GetHandle();
+}
+
+TOnlineAsyncOpHandle<FAuthLogout> FAuthCommon::Logout(FAuthLogout::Params&& Params)
+{
+	TOnlineAsyncOp<FAuthLogout>& Operation = Services.OpCache.GetOp<FAuthLogout>(MoveTemp(Params));
+	Operation.SetError(Errors::Unimplemented());
+	return Operation.GetHandle();
+}
+
+TOnlineAsyncOpHandle<FAuthGenerateAuth> FAuthCommon::GenerateAuth(FAuthGenerateAuth::Params&& Params)
+{
+	TOnlineAsyncOp<FAuthGenerateAuth>& Operation = Services.OpCache.GetOp<FAuthGenerateAuth>(MoveTemp(Params));
+	Operation.SetError(Errors::Unimplemented());
+	return Operation.GetHandle();
+}
+
+TOnlineEvent<void(const FLoginStatusChanged&)> FAuthCommon::OnLoginStatusChanged()
+{
+	return OnLoginStatusChangedEvent;
+}
+
+/* UE::Online */ }
diff --git a/Engine/Plugins/Online/Experimental/OnlineServices/Source/OnlineServicesCommon/Private/Online/FriendsCommon.cpp b/Engine/Plugins/Online/Experimental/OnlineServices/Source/OnlineServicesCommon/Private/Online/FriendsCommon.cpp
new file mode 100644
--- /dev/null
+++ b/Engine/Plugins/Online/Experimental/OnlineServices/Source/OnlineServicesCommon/Private/Online/FriendsCommon.cpp
@@ -0,0 +1,32 @@
+// Copyright Epic Games, Inc. All Rights Reserved.
+
+#include "Online/FriendsCommon.h"
+
+#include "Online/OnlineServicesCommon.h"
+
+namespace UE::Online {
+
+FFriendsCommon::FFriendsCommon(FOnlineServicesCommon& InServices)
+	: TOnlineComponent(TEXT("Friends"), InServices.AsShared())
+	, Services(InServices)
+{
+}
+
+TOnlineAsyncOpHandle<FQueryFriends> FFriendsCommon::QueryFriends(FQueryFriends::Params&& Params)
+{
+	TOnlineAsyncOp<FQueryFriends>& Operation = Services.OpCache.GetOp<FQueryFriends>(MoveTemp(Params));
+	Operation.SetError(Errors::Unimplemented());
+	return Operation.GetHandle();
+}
+
+TOnlineResult<FGetFriends::Result> FFriendsCommon::GetFriends(FGetFriends::Params&& Params)
+{
+	return TOnlineResult<FGetFriends::Result>(Errors::Unimplemented());
+}
+
+TOnlineEvent<void(const FFriendsListUpdated&)> FFriendsCommon::OnFriendsListUpdated()
+{
+	return OnFriendsListUpdatedEvent;
+}
+
+/* UE::Online */ }
diff --git a/Engine/Plugins/Online/Experimental/OnlineServices/Source/OnlineServicesCommon/Private/Online/OnlineServicesCommon.cpp b/Engine/Plugins/Online/Experimental/OnlineServices/Source/OnlineServicesCommon/Private/Online/OnlineServicesCommon.cpp
new file mode 100644
--- /dev/null
+++ b/Engine/Plugins/Online/Experimental/OnlineServices/Source/OnlineServicesCommon/Private/Online/OnlineServicesCommon.cpp
@@ -0,0 +1,75 @@
+// Copyright Epic Games, Inc. All Rights Reserved.
+
+#include "Online/OnlineServicesCommon.h"
+
+#include "Online/AuthCommon.h"
+#include "Online/FriendsCommon.h"
+
+namespace UE::Online {
+
+FOnlineServicesCommon::FOnlineServicesCommon()
+	: OpCache(TEXT("OnlineServicesCommon"), *this)
+	, ConfigProvider(MakeUnique<FOnlineConfigProviderGConfig>(GEngineIni))
+{
+}
+
+void FOnlineServicesCommon::Init()
+{
+	RegisterComponents();
+	Initialize();
+	PostInitialize();
+}
+
+void FOnlineServicesCommon::Destroy()
+{
+	PreShutdown();
+	Shutdown();
+}
+
+IAuthPtr FOnlineServicesCommon::GetAuthInterface()
+{
+	return IAuthPtr(AsShared(), Get<IAuth>());
+}
+
+IFriendsPtr FOnlineServicesCommon::GetFriendsInterface()
+{
+	return IFriendsPtr(AsShared(), Get<IFriends>());
+}
+
+void FOnlineServicesCommon::RegisterComponents()
+{
+}
+
+void FOnlineServicesCommon::Initialize()
+{
+	Components.Visit(&IOnlineComponent::Initialize);
+}
+
+void FOnlineServicesCommon::PostInitialize()
+{
+	Components.Visit(&IOnlineComponent::PostInitialize);
+}
+
+void FOnlineServicesCommon::LoadConfig()
+{
+	Components.Visit(&IOnlineComponent::LoadConfig);
+}
+
+bool FOnlineServicesCommon::Tick(float DeltaSeconds)
+{
+	Components.Visit(&IOnlineComponent::Tick, DeltaSeconds);
+
+	return true;
+}
+
+void FOnlineServicesCommon::PreShutdown()
+{
+	Components.Visit(&IOnlineComponent::PreShutdown);
+}
+
+void FOnlineServicesCommon::Shutdown()
+{
+	Components.Visit(&IOnlineComponent::Shutdown);
+}
+
+/* UE::Online */ }
diff --git a/Engine/Plugins/Online/Experimental/OnlineServices/Source/OnlineServicesCommon/Private/Online/OnlineSyntaxTests.cpp b/Engine/Plugins/Online/Experimental/OnlineServices/Source/OnlineServicesCommon/Private/Online/OnlineSyntaxTests.cpp
new file mode 100644
--- /dev/null
+++ b/Engine/Plugins/Online/Experimental/OnlineServices/Source/OnlineServicesCommon/Private/Online/OnlineSyntaxTests.cpp
@@ -0,0 +1,291 @@
+// Copyright Epic Games, Inc. All Rights Reserved.
+
+#include "Online/OnlineServicesCommon.h"
+#include "Online/OnlineConfig.h"
+#include "Online/OnlineAsyncOp.h"
+#include "Online/OnlineId.h"
+
+#define TEST_ASYNC_OP_CONTINUATION_SYNTAX 1
+#define TEST_CONSTRUCT_DELEGATE_SYNTAX 1
+#define TEST_CONSTRUCT_DELEGATE_SYNTAX_UOBJECT 0
+#define TEST_ASYNCOPCACHE_SYNTAX 1
+#define TEST_CONFIG_SYNTAX 1
+
+namespace UE::Online {
+
+namespace Test {
+
+struct ITestInterface
+{
+public:
+};
+
+struct FTestInterface : ITestInterface
+{
+public:
+	using Super = ITestInterface;
+};
+
+struct FTestOp
+{
+	struct Params
+	{
+	};
+
+	struct Result
+	{
+	};
+};
+
+struct FJoinableTestOp
+{
+	struct Params
+	{
+		FAccountId LocalUserId;
+	};
+
+	struct Result
+	{
+	};
+};
+
+/* Test */ }
+
+namespace Meta {
+
+BEGIN_ONLINE_STRUCT_META(Test::FTestOp::Params)
+END_ONLINE_STRUCT_META()
+
+BEGIN_ONLINE_STRUCT_META(Test::FJoinableTestOp::Params)
+	ONLINE_STRUCT_FIELD(Test::FJoinableTestOp::Params, LocalUserId)
+END_ONLINE_STRUCT_META()
+
+/* Meta */ }
+
+// confirm that a couple concepts are working as expected
+static_assert(TModels<UE::Online::Meta::CSuperDefined, Test::FTestInterface>::Value);
+static_assert(!TModels<UE::Online::Meta::CSuperDefined, Test::ITestInterface>::Value);
+static_assert(!TModels<UE::Online::CLocalUserIdDefined, Test::FTestOp::Params>::Value);
+static_assert(TModels<UE::Online::CLocalUserIdDefined, Test::FJoinableTestOp::Params>::Value);
+
+/* UE::Online */ }
+
+
+#if TEST_ASYNC_OP_CONTINUATION_SYNTAX
+
+namespace UE::Online {
+
+struct FTestOp
+{
+	struct Params {};
+	struct Result {};
+};
+
+inline void TestAsyncOpContinuationSyntax(FOnlineServicesCommon& Services)
+{
+	TSharedRef<TOnlineAsyncOp<FTestOp>> SharedOp = MakeShared<TOnlineAsyncOp<FTestOp>>(Services, FTestOp::Params());
+	TOnlineAsyncOp<FTestOp>& Op = SharedOp.Get();
+	TOnlineChainableAsyncOp<FTestOp, void> ChainableOp = Op;
+	if (true /* condition to enable optional steps */)
+	{
+		ChainableOp = ChainableOp.Then([](FOnlineAsyncOp&) {})
+			.Then([](FOnlineAsyncOp&) { return MakeFulfilledPromise<void>().GetFuture(); });
+	}
+	ChainableOp.Then([](FOnlineAsyncOp&) {})
+		.Then([](FOnlineAsyncOp&) { return MakeFulfilledPromise<void>().GetFuture(); })
+		.Then([](FOnlineAsyncOp&) {})
+		.Then([](TOnlineAsyncOp<FTestOp>&) { return 1; })
+		.Then([](TOnlineAsyncOp<FTestOp>&, int) { return MakeFulfilledPromise<FString>().GetFuture(); })
+		.Then([](TOnlineAsyncOp<FTestOp>&, FString) { return FString(); })
+		.Then([](TOnlineAsyncOp<FTestOp>&, const FString&) { return FString(); })
+		.Then([](TOnlineAsyncOp<FTestOp>&, FString&&) {})
+		.Then([](TOnlineAsyncOp<FTestOp>&) { return 1.0f; })
+		.Then([](TOnlineAsyncOp<FTestOp>&, float) { return MakeFulfilledPromise<void>().GetFuture(); })
+		.Then([](TOnlineAsyncOp<FTestOp>&) { return MakeFulfilledPromise<void>().GetFuture(); })
+		.Enqueue(/*queue*/);
+	Op.GetHandle();
+}
+
+/* UE::Online */ }
+
+#endif // TEST_ASYNC_OP_CONTINUATION_SYNTAX
+
+
+#if TEST_CONSTRUCT_DELEGATE_SYNTAX
+
+class FSPTest : public TSharedFromThis<FSPTest, ESPMode::NotThreadSafe>
+{
+public:
+	void FN(int, const FString&) {}
+	void FN2(int, const FString&, float) {}
+	void CFN(int, const FString&) const {}
+	void CFN2(int, const FString&, float) const {}
+};
+
+class FTSSPTest : public TSharedFromThis<FTSSPTest, ESPMode::ThreadSafe>
+{
+public:
+	void FN(int, const FString&) {}
+	void FN2(int, const FString&, float) {}
+	void CFN(int, const FString&) const {}
+	void CFN2(int, const FString&, float) const {}
+};
+
+class FRawTest
+{
+public:
+	void FN(int, const FString&) {}
+	void FN2(int, const FString&, float) {}
+	void CFN(int, const FString&) const {}
+	void CFN2(int, const FString&, float) const {}
+};
+
+#if TEST_CONSTRUCT_DELEGATE_SYNTAX_UOBJECT
+#include "UObject/Object.h"
+#include "UObject/WeakObjectPtr.h"
+
+class UMyObject : public UObject
+{
+public:
+	void FN(int, const FString&) {}
+	void FN2(int, const FString&, float) {}
+	void CFN(int, const FString&) const {}
+	void CFN2(int, const FString&, float) const {}
+};
+#endif // TEST_CONSTRUCT_DELEGATE_SYNTAX_UOBJECT
+
+inline void FN(int, const FString&) {}
+inline void FN2(int, const FString&, float) {}
+
+inline void DelegateTest()
+{
+	using UE::Online::Private::ConstructDelegate;
+
+	TSharedRef<FSPTest, ESPMode::NotThreadSafe> SP = MakeShared<FSPTest, ESPMode::NotThreadSafe>();
+	TSharedRef<const FSPTest, ESPMode::NotThreadSafe> ConstSP = ConstCastSharedRef<const FSPTest>(SP);
+	TSharedRef<FTSSPTest, ESPMode::ThreadSafe> TSSP = MakeShared<FTSSPTest, ESPMode::ThreadSafe>();
+	TSharedRef<const FTSSPTest, ESPMode::ThreadSafe> ConstTSSP = ConstCastSharedRef<const FTSSPTest>(TSSP);
+	FRawTest* Raw = new FRawTest();
+	const FRawTest* ConstRaw = Raw;
+#if TEST_CONSTRUCT_DELEGATE_SYNTAX_UOBJECT
+	UMyObject* UObject = NewObject<UMyObject>();
+	const UMyObject* ConstUObject = UObject;
+#endif // TEST_CONSTRUCT_DELEGATE_SYNTAX_UOBJECT
+
+	// CreateStatic
+	TDelegate<void(int, const FString&)> StaticDelegate1 = ConstructDelegate<void(int, const FString&)>(FN);
+	TDelegate<void(int, const FString&)> StaticDelegate2 = ConstructDelegate<void(int, const FString&)>(&FN);
+	TDelegate<void(int, const FString&)> StaticDelegate3 = ConstructDelegate<void(int, const FString&)>(FN2, 1.0f);
+	TDelegate<void(int, const FString&)> StaticDelegate4 = ConstructDelegate<void(int, const FString&)>(&FN2, 1.0f);
+	// CreateLambda
+	TDelegate<void(int, const FString&)> LambdaDelegate = ConstructDelegate<void(int, const FString&)>([](int, const FString&) {});
+	// CreateUObject
+#if TEST_CONSTRUCT_DELEGATE_SYNTAX_UOBJECT
+	TDelegate<void(int, const FString&)> UObjectDelegate1 = ConstructDelegate<void(int, const FString&)>(UObject, &UMyObject::FN);
+	TDelegate<void(int, const FString&)> UObjectDelegate2 = ConstructDelegate<void(int, const FString&)>(UObject, &UMyObject::FN2, 1.0f);
+	TDelegate<void(int, const FString&)> UObjectDelegate3 = ConstructDelegate<void(int, const FString&)>(UObject, &UMyObject::CFN);
+	TDelegate<void(int, const FString&)> UObjectDelegate4 = ConstructDelegate<void(int, const FString&)>(UObject, &UMyObject::CFN2, 1.0f);
+	TDelegate<void(int, const FString&)> UObjectDelegate5 = ConstructDelegate<void(int, const FString&)>(ConstUObject, &UMyObject::CFN);
+	TDelegate<void(int, const FString&)> UObjectDelegate6 = ConstructDelegate<void(int, const FString&)>(ConstUObject, &UMyObject::CFN2, 1.0f);
+	// CreateUFunction
+	TDelegate<void(int, const FString&)> UFunctionDelegate1 = ConstructDelegate<void(int, const FString&)>(UObject, FName(TEXT("FN")));
+	TDelegate<void(int, const FString&)> UFunctionDelegate2 = ConstructDelegate<void(int, const FString&)>(UObject, FName(TEXT("FN2")), 1.0f);
+#endif // TEST_CONSTRUCT_DELEGATE_SYNTAX_UOBJECT
+	// CreateThreadSafeSP
+	TDelegate<void(int, const FString&)> TSSPDelegate1 = ConstructDelegate<void(int, const FString&)>(TSSP, &FTSSPTest::FN);
+	TDelegate<void(int, const FString&)> TSSPDelegate2 = ConstructDelegate<void(int, const FString&)>(TSSP, &FTSSPTest::FN2, 1.0f);
+	TDelegate<void(int, const FString&)> TSSPDelegate3 = ConstructDelegate<void(int, const FString&)>(&TSSP.Get(), &FTSSPTest::FN);
+	TDelegate<void(int, const FString&)> TSSPDelegate4 = ConstructDelegate<void(int, const FString&)>(&TSSP.Get(), &FTSSPTest::FN2, 1.0f);
+	TDelegate<void(int, const FString&)> TSSPDelegate5 = ConstructDelegate<void(int, const FString&)>(TSSP, &FTSSPTest::CFN);
+	TDelegate<void(int, const FString&)> TSSPDelegate6 = ConstructDelegate<void(int, const FString&)>(TSSP, &FTSSPTest::CFN2, 1.0f);
+	TDelegate<void(int, const FString&)> TSSPDelegate7 = ConstructDelegate<void(int, const FString&)>(&TSSP.Get(), &FTSSPTest::CFN);
+	TDelegate<void(int, const FString&)> TSSPDelegate8 = ConstructDelegate<void(int, const FString&)>(&TSSP.Get(), &FTSSPTest::CFN2, 1.0f);
+	TDelegate<void(int, const FString&)> TSSPDelegate9 = ConstructDelegate<void(int, const FString&)>(ConstTSSP, &FTSSPTest::CFN);
+	TDelegate<void(int, const FString&)> TSSPDelegate10 = ConstructDelegate<void(int, const FString&)>(ConstTSSP, &FTSSPTest::CFN2, 1.0f);
+	TDelegate<void(int, const FString&)> TSSPDelegate11 = ConstructDelegate<void(int, const FString&)>(&ConstTSSP.Get(), &FTSSPTest::CFN);
+	TDelegate<void(int, const FString&)> TSSPDelegate12 = ConstructDelegate<void(int, const FString&)>(&ConstTSSP.Get(), &FTSSPTest::CFN2, 1.0f);
+	// CreateSP
+	TDelegate<void(int, const FString&)> SPDelegate1 = ConstructDelegate<void(int, const FString&)>(SP, &FSPTest::FN);
+	TDelegate<void(int, const FString&)> SPDelegate2 = ConstructDelegate<void(int, const FString&)>(SP, &FSPTest::FN2, 1.0f);
+	TDelegate<void(int, const FString&)> SPDelegate3 = ConstructDelegate<void(int, const FString&)>(&SP.Get(), &FSPTest::FN);
+	TDelegate<void(int, const FString&)> SPDelegate4 = ConstructDelegate<void(int, const FString&)>(&SP.Get(), &FSPTest::FN2, 1.0f);
+	TDelegate<void(int, const FString&)> SPDelegate5 = ConstructDelegate<void(int, const FString&)>(SP, &FSPTest::CFN);
+	TDelegate<void(int, const FString&)> SPDelegate6 = ConstructDelegate<void(int, const FString&)>(SP, &FSPTest::CFN2, 1.0f);
+	TDelegate<void(int, const FString&)> SPDelegate7 = ConstructDelegate<void(int, const FString&)>(&SP.Get(), &FSPTest::CFN);
+	TDelegate<void(int, const FString&)> SPDelegate8 = ConstructDelegate<void(int, const FString&)>(&SP.Get(), &FSPTest::CFN2, 1.0f);
+	TDelegate<void(int, const FString&)> SPDelegate9 = ConstructDelegate<void(int, const FString&)>(ConstSP, &FSPTest::CFN);
+	TDelegate<void(int, const FString&)> SPDelegate10 = ConstructDelegate<void(int, const FString&)>(ConstSP, &FSPTest::CFN2, 1.0f);
+	TDelegate<void(int, const FString&)> SPDelegate11 = ConstructDelegate<void(int, const FString&)>(&ConstSP.Get(), &FSPTest::CFN);
+	TDelegate<void(int, const FString&)> SPDelegate12 = ConstructDelegate<void(int, const FString&)>(&ConstSP.Get(), &FSPTest::CFN2, 1.0f);
+	// CreateRaw
+	TDelegate<void(int, const FString&)> RawDelegate1 = ConstructDelegate<void(int, const FString&)>(Raw, &FRawTest::FN);
+	TDelegate<void(int, const FString&)> RawDelegate2 = ConstructDelegate<void(int, const FString&)>(Raw, &FRawTest::FN2, 1.0f);
+	TDelegate<void(int, const FString&)> RawDelegate3 = ConstructDelegate<void(int, const FString&)>(Raw, &FRawTest::CFN);
+	TDelegate<void(int, const FString&)> RawDelegate4 = ConstructDelegate<void(int, const FString&)>(Raw, &FRawTest::CFN2, 1.0f);
+	TDelegate<void(int, const FString&)> RawDelegate5 = ConstructDelegate<void(int, const FString&)>(ConstRaw, &FRawTest::CFN);
+	TDelegate<void(int, const FString&)> RawDelegate6 = ConstructDelegate<void(int, const FString&)>(ConstRaw, &FRawTest::CFN2, 1.0f);
+	// CreateWeakLambda
+#if TEST_CONSTRUCT_DELEGATE_SYNTAX_UOBJECT
+	TDelegate<void(int, const FString&)> WeakUObjectDelegate1 = ConstructDelegate<void(int, const FString&)>(UObject, [](int, const FString&) {});
+//	TDelegate<void(int, const FString&)> WeakUObjectDelegate2 = ConstructDelegate<void(int, const FString&)>(ConstUObject, [](int, const FString&) {}); // CreateWeakObject doesn't support const UObject*
+#endif // TEST_CONSTRUCT_DELEGATE_SYNTAX_UOBJECT
+	// lambda with TWeakPtr
+	TDelegate<void(int, const FString&)> WeakTSSPDelegate1 = ConstructDelegate<void(int, const FString&)>(TSSP, [](int, const FString&) {});
+	TDelegate<void(int, const FString&)> WeakTSSPDelegate2 = ConstructDelegate<void(int, const FString&)>(&TSSP.Get(), [](int, const FString&) {});
+	TDelegate<void(int, const FString&)> WeakSPDelegate1 = ConstructDelegate<void(int, const FString&)>(SP, [](int, const FString&) {});
+	TDelegate<void(int, const FString&)> WeakSPDelegate2 = ConstructDelegate<void(int, const FString&)>(&SP.Get(), [](int, const FString&) {});
+	// lambda with raw pointer: Error: When using non pointer to member functions, the first parameter can only be a UObject*, TSharedRef, or pointer to a class that derives from TSharedFromThis
+//	TDelegate<void(int, const FString&)> WeakRawDelegate = ConstructDelegate<void(int, const FString&)>(Raw, [](int, const FString&) {});
+
+	delete Raw;
+}
+
+#endif // TEST_CONSTRUCT_DELEGATE_SYNTAX
+
+
+#if TEST_ASYNCOPCACHE_SYNTAX
+
+namespace UE::Online {
+
+inline void TestAsyncOpCacheSyntax(FOnlineServicesCommon& Services)
+{
+	Test::FTestOp::Params Params;
+	TOnlineAsyncOp<Test::FTestOp>& Op = Services.OpCache.GetOp<Test::FTestOp>(MoveTemp(Params));
+
+	Test::FJoinableTestOp::Params JoinableParams;
+	TOnlineAsyncOp<Test::FJoinableTestOp>& Op2 = Services.OpCache.GetJoinableOp<Test::FJoinableTestOp>(MoveTemp(JoinableParams));
+}
+
+/* UE::Online */ }
+
+#endif // TEST_ASYNCOPCACHE_SYNTAX
+
+
+#if TEST_CONFIG_SYNTAX
+
+struct FTestConfig
+{
+	float Float;
+	int Int;
+	FString String;
+	TArray<int> IntArray;
+};
+
+namespace UE::Online::Meta {
+
+BEGIN_ONLINE_STRUCT_META(FTestConfig)
+	ONLINE_STRUCT_FIELD(FTestConfig, Float),
+	ONLINE_STRUCT_FIELD(FTestConfig, Int),
+	ONLINE_STRUCT_FIELD(FTestConfig, String),
+	ONLINE_STRUCT_FIELD(FTestConfig, IntArray)
+END_ONLINE_STRUCT_META()
+
+/* UE::Online::Meta*/ }
+
+void TestLoadConfigSyntax(UE::Online::IOnlineConfigProvider& ConfigProvider)
+{
+	FTestConfig TestConfig;
+	UE::Online::LoadConfig(ConfigProvider, TEXT("Test"), TestConfig);
+}
+
+#endif // TEST_CONFIG_SYNTAX
diff --git a/Engine/Plugins/Online/Experimental/OnlineServices/Source/OnlineServicesCommon/Private/OnlineServicesCommonModule.cpp b/Engine/Plugins/Online/Experimental/OnlineServices/Source/OnlineServicesCommon/Private/OnlineServicesCommonModule.cpp
new file mode 100644
--- /dev/null
+++ b/Engine/Plugins/Online/Experimental/OnlineServices/Source/OnlineServicesCommon/Private/OnlineServicesCommonModule.cpp
@@ -0,0 +1,20 @@
+// Copyright Epic Games, Inc. All Rights Reserved.
+
+#include "CoreMinimal.h"
+#include "Modules/ModuleManager.h"
+
+class FOnlineServicesCommon : public IModuleInterface
+{
+public:
+	virtual void StartupModule()
+	{
+
+	}
+
+	virtual void ShutdownModule()
+	{
+
+	}
+};
+
+IMPLEMENT_MODULE(FOnlineServicesCommon, OnlineServicesCommon);
diff --git a/Engine/Plugins/Online/Experimental/OnlineServices/Source/OnlineServicesCommon/Public/Online/AuthCommon.h b/Engine/Plugins/Online/Experimental/OnlineServices/Source/OnlineServicesCommon/Public/Online/AuthCommon.h
new file mode 100644
--- /dev/null
+++ b/Engine/Plugins/Online/Experimental/OnlineServices/Source/OnlineServicesCommon/Public/Online/AuthCommon.h
@@ -0,0 +1,30 @@
+// Copyright Epic Games, Inc. All Rights Reserved.
+
+#pragma once
+
+#include "Online/Auth.h"
+#include "Online/OnlineComponent.h"
+
+namespace UE::Online {
+
+class FOnlineServicesCommon;
+
+class ONLINESERVICESCOMMON_API FAuthCommon : public TOnlineComponent<IAuth>
+{
+public:
+	using Super = IAuth;
+
+	FAuthCommon(FOnlineServicesCommon& InServices);
+
+	virtual TOnlineAsyncOpHandle<FAuthLogin> Login(FAuthLogin::Params&& Params) override;
+	virtual TOnlineAsyncOpHandle<FAuthLogout> Logout(FAuthLogout::Params&& Params) override;
+	virtual TOnlineAsyncOpHandle<FAuthGenerateAuth> GenerateAuth(FAuthGenerateAuth::Params&& Params) override;
+	virtual TOnlineEvent<void(const FLoginStatusChanged&)> OnLoginStatusChanged() override;
+
+protected:
+	FOnlineServicesCommon& Services;
+
+	TOnlineEventCallable<void(const FLoginStatusChanged&)> OnLoginStatusChangedEvent;
+};
+
+/* UE::Online */ }
diff --git a/Engine/Plugins/Online/Experimental/OnlineServices/Source/OnlineServicesCommon/Public/Online/FriendsCommon.h b/Engine/Plugins/Online/Experimental/OnlineServices/Source/OnlineServicesCommon/Public/Online/FriendsCommon.h
new file mode 100644
--- /dev/null
+++ b/Engine/Plugins/Online/Experimental/OnlineServices/Source/OnlineServicesCommon/Public/Online/FriendsCommon.h
@@ -0,0 +1,30 @@
+// Copyright Epic Games, Inc. All Rights Reserved.
+
+#pragma once
+
+#include "Online/Friends.h"
+#include "Online/OnlineComponent.h"
+
+namespace UE::Online {
+
+class FOnlineServicesCommon;
+
+class ONLINESERVICESCOMMON_API FFriendsCommon : public TOnlineComponent<IFriends>
+{
+public:
+	using Super = IFriends;
+
+	FFriendsCommon(FOnlineServicesCommon& InServices);
+
+	// IFriends
+	virtual TOnlineAsyncOpHandle<FQueryFriends> QueryFriends(FQueryFriends::Params&& Params) override;
+	virtual TOnlineResult<FGetFriends::Result> GetFriends(FGetFriends::Params&& Params) override;
+	virtual TOnlineEvent<void(const FFriendsListUpdated&)> OnFriendsListUpdated() override;
+
+protected:
+	FOnlineServicesCommon& Services;
+
+	TOnlineEventCallable<void(const FFriendsListUpdated&)> OnFriendsListUpdatedEvent;
+};
+
+/* UE::Online */ }
diff --git a/Engine/Plugins/Online/Experimental/OnlineServices/Source/OnlineServicesCommon/Public/Online/OnlineAsyncOp.h b/Engine/Plugins/Online/Experimental/OnlineServices/Source/OnlineServicesCommon/Public/Online/OnlineAsyncOp.h
new file mode 100644
--- /dev/null
+++ b/Engine/Plugins/Online/Experimental/OnlineServices/Source/OnlineServicesCommon/Public/Online/OnlineAsyncOp.h
@@ -0,0 +1,545 @@
+// Copyright Epic Games, Inc. All Rights Reserved.
+
+#pragma once
+
+#include "Online/OnlineAsyncOpHandle.h"
+#include "Online/OnlineTypeInfo.h"
+
+#include "Containers/Map.h"
+#include "Templates/UniquePtr.h"
+#include "Async/Future.h"
+
+namespace UE::Online {
+
+class FOnlineServicesCommon;
+
+// TEMP
+namespace Errors
+{
+	inline FOnlineError Unknown() { return FOnlineError();  }
+}
+// END TEMP
+
+template <typename OpType> class TOnlineAsyncOp;
+template <typename OpType, typename T> class TOnlineChainableAsyncOp;
+
+enum class EOnlineAsyncExecutionPolicy
+{
+	RunOnGameThread
+	// TODO: Add/implement others
+};
+
+// TODO
+class FOnlineAsyncExecutionPolicy
+{
+public:
+	FOnlineAsyncExecutionPolicy(EOnlineAsyncExecutionPolicy InExecutionPolicy)
+		: ExecutionPolicy(InExecutionPolicy)
+	{
+	}
+
+	static FOnlineAsyncExecutionPolicy RunOnGameThread() { return FOnlineAsyncExecutionPolicy(EOnlineAsyncExecutionPolicy::RunOnGameThread); }
+
+	EOnlineAsyncExecutionPolicy GetExecutionPolicy() const { return ExecutionPolicy; }
+
+private:
+	EOnlineAsyncExecutionPolicy ExecutionPolicy;
+};
+
+
+namespace Private
+{
+
+class FOnlineOperationData // Map of (TypeName,Key)->(data of any Type)
+{
+public:
+	template <typename T>
+	void Set(const FString& Key, T&& InData)
+	{
+		Data.Add(FOperationDataKey{ TOnlineTypeInfo<T>::GetTypeName(), Key }, MakeUnique<TData<T>>(MoveTemp(InData)));
+	}
+
+	template <typename T>
+	void Set(const FString& Key, const T& InData)
+	{
+		Data.Add(FOperationDataKey{ TOnlineTypeInfo<T>::GetTypeName(), Key }, MakeUnique<TData<T>>(InData));
+	}
+
+	template <typename T>
+	const T* Get(const FString& Key) const
+	{
+		if (auto Value = Data.Find(FOperationDataKey{ TOnlineTypeInfo<T>::GetTypeName(), Key }))
+		{
+			return static_cast<const T*>((*Value)->GetData());
+		}
+
+		return nullptr;
+	}
+
+	struct FOperationDataKey
+	{
+		FOnlineTypeName TypeName;
+		FString Key;
+
+		bool operator==(const FOperationDataKey& Other) const
+		{
+			return TypeName == Other.TypeName && Key == Other.Key;
+		}
+	};
+
+private:
+	class IData
+	{
+	public:
+		virtual ~IData() {}
+		virtual FOnlineTypeName GetTypeName() = 0;
+		virtual void* GetData() = 0;
+
+		template <typename T>
+		const T* Get()
+		{
+			if (GetTypeName() == TOnlineTypeInfo<T>::GetTypeName())
+			{
+				return static_cast<T*>(GetData());
+			}
+
+			return nullptr;
+		}
+	};
+
+	template <typename T>
+	class TData : public IData
+	{
+	public:
+		TData(const T& InData)
+			: Data(InData)
+		{
+		}
+
+		TData(T&& InData)
+			: Data(MoveTemp(InData))
+		{
+		}
+
+		virtual FOnlineTypeName GetTypeName() override
+		{
+			return  TOnlineTypeInfo<T>::GetTypeName();
+		}
+
+		virtual void* GetData() override
+		{
+			return &Data;
+		}
+
+	private:
+		T Data;
+	};
+
+	friend uint32 GetTypeHash(const FOperationDataKey& Key);
+
+	TMap<FOperationDataKey, TUniquePtr<IData>> Data;
+};
+
+inline uint32 GetTypeHash(const FOnlineOperationData::FOperationDataKey& Key)
+{
+	using ::GetTypeHash;
+	return HashCombine(GetTypeHash(Key.TypeName), GetTypeHash(Key.Key));
+}
+
+template <typename T>
+struct TUnwrapFuture
+{
+	using Type = T;
+};
+
+template <typename T>
+struct TUnwrapFuture<TFuture<T>>
+{
+	using Type = T;
+};
+
+template <typename T>
+using TUnwrapFuture_T = typename TUnwrapFuture<T>::Type;
+
+template<typename ResultType>
+inline void SetPromiseValue(TPromise<ResultType>& Promise, TFuture<ResultType>&& Future)
+{
+	Promise.SetValue(Future.Get());
+}
+
+template <typename T>
+inline void SetPromiseValue(TPromise<void>& Promise, TFuture<T>&& Future)
+{
+	Promise.SetValue();
+}
+
+template<typename Func, typename... ParamTypes, typename ResultType>
+inline void SetPromiseValueFromCallable(TPromise<ResultType>& Promise, Func& Function, ParamTypes&&... Params)
+{
+	Promise.SetValue(Function(Forward<ParamTypes>(Params)...));
+}
+template<typename Func, typename... ParamTypes>
+inline void SetPromiseValueFromCallable(TPromise<void>& Promise, Func& Function, ParamTypes&&... Params)
+{
+	Function(Forward<ParamTypes>(Params)...);
+	Promise.SetValue();
+}
+
+template <typename T>
+struct TAsyncOpHelper
+{
+	template<typename Func, typename AsyncOpType, typename FutureValueType>
+	inline static auto Call(Func& Function, AsyncOpType& Op, TFuture<FutureValueType>& Value) -> decltype(Invoke(Function, Op, Value.Get()))
+	{
+		return Invoke(Function, Op, Value.Get());
+	}
+};
+
+template <>
+struct TAsyncOpHelper<void>
+{
+	template<typename Func, typename AsyncOpType, typename FutureValueType>
+	inline static auto Call(Func& Function, AsyncOpType& Op, TFuture<FutureValueType>& Value) -> decltype(Invoke(Function, Op))
+	{
+		return Invoke(Function, Op);
+	}
+};
+
+template <typename CallableType, typename... ParamTypes>
+struct TChainableAsyncOpCallableResultHelper
+{
+	using Type = TInvokeResult_T<CallableType, ParamTypes...>;
+};
+
+template <typename CallableType, typename OpType>
+struct TChainableAsyncOpCallableResultHelper<CallableType, TOnlineAsyncOp<OpType>&, void>
+{
+	using Type = TInvokeResult_T<CallableType, TOnlineAsyncOp<OpType>&>;
+};
+
+template <typename CallableType, typename... ParamTypes>
+using TChainableAsyncOpCallableResultHelper_T = typename TChainableAsyncOpCallableResultHelper<CallableType, ParamTypes...>::Type;
+
+template <typename OpType, typename CallableType, typename LastResultType>
+TOnlineChainableAsyncOp<OpType, Private::TUnwrapFuture_T<TChainableAsyncOpCallableResultHelper_T<CallableType, TOnlineAsyncOp<OpType>&, LastResultType>>> CreateChainableAsyncOp(TOnlineAsyncOp<OpType>& Op, TFuture<LastResultType>& LastResult, CallableType&& InCallable, FOnlineAsyncExecutionPolicy ExecutionPolicy)
+{
+	using ReturnType = TChainableAsyncOpCallableResultHelper_T<CallableType, TOnlineAsyncOp<OpType>&, LastResultType>;
+	using UnwrappedReturnType = Private::TUnwrapFuture_T<ReturnType>;
+
+	TPromise<UnwrappedReturnType> Promise;
+	TFuture<UnwrappedReturnType> Future = Promise.GetFuture();
+
+	using LastResultContinuationFutureType = std::conditional_t<std::is_same_v<void, LastResultType>, TFuture<int>, TFuture<LastResultType>>;
+	LastResult.Then([ExecutionPolicy, Callable = MoveTemp(InCallable), WeakOp = TWeakPtr<TOnlineAsyncOp<OpType>>(Op.AsShared()), MovedPromise = MoveTemp(Promise)](LastResultContinuationFutureType&& Result) mutable // mutable so we can set the promise value
+	{
+		TSharedPtr<TOnlineAsyncOp<OpType>> PinnedOp = WeakOp.Pin();
+		if (PinnedOp && !PinnedOp->IsComplete())
+		{
+			PinnedOp->GetServices().Execute(ExecutionPolicy, [Callable2 = MoveTemp(Callable), WeakOp, MovedPromise2 = MoveTemp(MovedPromise), MovedResult = MoveTemp(Result)]() mutable // mutable so we can set the promise value
+			{
+				TSharedPtr<TOnlineAsyncOp<OpType>> PinnedOp2 = WeakOp.Pin();
+				if (PinnedOp2 && !PinnedOp2->IsComplete())
+				{
+					if constexpr (std::is_same_v<ReturnType, TFuture<UnwrappedReturnType>>) // return type is a future
+					{
+						using ContinuationFutureType = std::conditional_t<std::is_same_v<void, UnwrappedReturnType>, TFuture<int>, TFuture<UnwrappedReturnType>>; // TFuture<void> continuation uses a TFuture<int>
+						TAsyncOpHelper<LastResultType>::Call(Callable2, *PinnedOp2, MovedResult).Then([WeakOp, MovedPromise3 = MoveTemp(MovedPromise2)](ContinuationFutureType&& Value) mutable // mutable so we can set the promise value
+						{
+							TSharedPtr<TOnlineAsyncOp<OpType>> PinnedOp3 = WeakOp.Pin();
+							if (PinnedOp3 && !PinnedOp3->IsComplete())
+							{
+								Private::SetPromiseValue(MovedPromise3, MoveTemp(Value));
+							}
+						});
+					}
+					else
+					{
+						if constexpr (std::is_same_v<void, LastResultType>)
+						{
+							Private::SetPromiseValueFromCallable(MovedPromise2, Callable2, *PinnedOp2);
+						}
+						else
+						{
+							Private::SetPromiseValueFromCallable(MovedPromise2, Callable2, *PinnedOp2, Forward<LastResultType>(MovedResult.Get()));
+						}
+					}
+				}
+			});
+		}
+	});
+
+	return TOnlineChainableAsyncOp<OpType, UnwrappedReturnType>(Op, MoveTemp(Future));
+}
+
+/* Private */ }
+
+template <typename OpType, typename T>
+class TOnlineChainableAsyncOp
+{
+public:
+	TOnlineChainableAsyncOp(TOnlineAsyncOp<OpType>& InOwningOperation, TFuture<T>&& InLastResult)
+		: OwningOperation(InOwningOperation)
+		, LastResult(MoveTemp(InLastResult))
+	{
+	}
+
+	TOnlineChainableAsyncOp(TOnlineChainableAsyncOp&& Other)
+		: OwningOperation(Other.OwningOperation)
+		, LastResult(MoveTemp(Other.LastResult))
+	{
+	}
+
+	TOnlineChainableAsyncOp& operator=(TOnlineChainableAsyncOp&& Other)
+	{
+		check(&OwningOperation == &Other.OwningOperation); // Can't reassign this
+		LastResult = MoveTemp(Other.LastResult);
+		return *this;
+	}
+
+	// Callable that takes a FOnlineAsyncOp& or TOnlineAsyncOp<OpType>& and a const T& and returns a Type or TFuture<Type>
+	template <typename CallableType>
+	TOnlineChainableAsyncOp<OpType, Private::TUnwrapFuture_T<Private::TChainableAsyncOpCallableResultHelper_T<CallableType, TOnlineAsyncOp<OpType>&, T>>>
+		Then(CallableType&& InCallable, FOnlineAsyncExecutionPolicy ExecutionPolicy = FOnlineAsyncExecutionPolicy::RunOnGameThread())
+	{
+		return Private::CreateChainableAsyncOp<OpType>(OwningOperation, LastResult, InCallable, ExecutionPolicy);
+	}
+
+	// placeholder
+	template <typename... U>
+	void Enqueue(U&&...)
+	{
+		static_assert(std::is_same_v<T, void>, "Continuation result discarded. Continuation prior to calling Enqueue must have a void or TFuture<void> return type.");
+		OwningOperation.Enqueue();
+	}
+
+protected:
+	TOnlineAsyncOp<OpType>& OwningOperation;
+	TFuture<T> LastResult;
+};
+
+class FOnlineAsyncOp
+{
+public:
+	virtual ~FOnlineAsyncOp() {}
+
+	Private::FOnlineOperationData Data;
+
+	virtual void SetError(FOnlineError&& Error) = 0;
+};
+
+// This class represents an async operation on the public interface
+// There may be one or more handles pointing to one instance
+template <typename OpType>
+class TOnlineAsyncOp 
+	: public FOnlineAsyncOp
+	, public TSharedFromThis<TOnlineAsyncOp<OpType>>
+{
+public:
+	using ParamsType = typename OpType::Params;
+	using ResultType = typename OpType::Result;
+
+	TOnlineAsyncOp(FOnlineServicesCommon& InServices, ParamsType&& Params)
+		: Services(InServices)
+		, SharedState(MakeShared<FAsyncOpSharedState>(MoveTemp(Params)))
+	{
+	}
+
+	~TOnlineAsyncOp()
+	{
+	}
+
+	bool IsReady() const
+	{
+		return SharedState->State != EAsyncOpState::Invalid;
+	}
+
+	bool IsComplete() const
+	{
+		return SharedState->State >= EAsyncOpState::Complete;
+	}
+
+	const ParamsType& GetParams() const
+	{
+		return SharedState->Params;
+	}
+
+	// Callable that takes a FOnlineAsyncOp& or TOnlineAsyncOp<OpType>& and returns a Type or TFuture<Type>
+	template <typename CallableType>
+	TOnlineChainableAsyncOp<OpType, Private::TUnwrapFuture_T<TInvokeResult_T<CallableType, TOnlineAsyncOp<OpType>&>>>
+		Then(CallableType&& InCallable, FOnlineAsyncExecutionPolicy ExecutionPolicy = FOnlineAsyncExecutionPolicy::RunOnGameThread())
+	{
+		TFuture<void> StartExecutionFuture = StartExecutionPromise.GetFuture();
+		return Private::CreateChainableAsyncOp<OpType>(*this, StartExecutionFuture, InCallable, ExecutionPolicy);
+	}
+
+	operator TOnlineChainableAsyncOp<OpType, void>()
+	{
+		return TOnlineChainableAsyncOp<OpType, void>(*this, StartExecutionPromise.GetFuture());
+	}
+
+	static TOnlineAsyncOp<OpType> CreateError(const FOnlineError& Error) { return TOnlineAsyncOp<OpType>(); }
+
+	TOnlineAsyncOpHandle<OpType> GetHandle()
+	{
+		return TOnlineAsyncOpHandle<OpType>(CreateSharedState());
+	}
+
+	void Cancel(const FOnlineError& Reason)
+	{
+		SetError(Reason);
+		SharedState->State = EAsyncOpState::Cancelled;
+	}
+
+	void SetResult(ResultType&& InResult)
+	{
+		SharedState->Result = TOnlineResult<ResultType>(MoveTemp(InResult));
+		SharedState->State = EAsyncOpState::Complete;
+
+		TriggerOnComplete(SharedState->Result);
+	}
+
+	virtual void SetError(FOnlineError&& Error)
+	{
+		SharedState->Result = TOnlineResult<ResultType>(MoveTemp(Error));
+		SharedState->State = EAsyncOpState::Complete;
+
+		TriggerOnComplete(SharedState->Result);
+	}
+
+	FOnlineServicesCommon& GetServices() { return Services; }
+
+	template <typename... U>
+	void Enqueue(U&&...)
+	{
+		// placeholder
+		StartExecutionPromise.SetValue();
+	}
+
+protected:
+	FOnlineServicesCommon& Services;
+
+	void TriggerOnComplete(const TOnlineResult<ResultType>& Result)
+	{
+		TArray<TSharedRef<FAsyncOpSharedHandleState>> SharedHandleStatesCopy(SharedHandleStates);
+		for (TSharedRef<FAsyncOpSharedHandleState>& SharedHandleState : SharedHandleStatesCopy)
+		{
+			SharedHandleState->TriggerOnComplete(Result);
+		}
+	}
+
+	class FAsyncOpSharedState
+	{
+	public:
+		FAsyncOpSharedState(ParamsType&& InParams)
+			: Params(MoveTemp(InParams))
+		{
+		}
+
+		ParamsType Params;
+		// This will need to be protected with a mutex if we want to allow this to be set from multiple threads (eg, set result from a task graph thread, while allowing this to be cancelled from the game thread)
+		TOnlineResult<ResultType> Result{ Errors::Unknown() };
+		EAsyncOpState State = EAsyncOpState::Invalid;
+
+		bool IsComplete() const
+		{
+			return State >= EAsyncOpState::Complete;
+		}
+	};
+
+	class FAsyncOpSharedHandleState : public Private::IOnlineAsyncOpSharedState<OpType>, public TSharedFromThis<FAsyncOpSharedHandleState>
+	{
+	public:
+		FAsyncOpSharedHandleState(const TSharedRef<TOnlineAsyncOp<OpType>>& InAsyncOp)
+			: SharedState(InAsyncOp->SharedState)
+			, AsyncOp(InAsyncOp)
+		{
+		}
+
+		~FAsyncOpSharedHandleState()
+		{
+			Detach();
+		}
+
+		virtual void Cancel(const FOnlineError& Reason) override
+		{
+			TSharedPtr<TOnlineAsyncOp<OpType>> PinnedOP = AsyncOp.Pin();
+			if (PinnedOP.IsValid())
+			{
+				bCancelled = true;
+				TriggerOnComplete(TOnlineResult<ResultType>(Reason));
+			}
+		}
+
+		virtual EAsyncOpState GetState() const override
+		{
+			return bCancelled ? EAsyncOpState::Cancelled : SharedState->State;
+		}
+
+		virtual void SetOnProgress(TDelegate<void(const FAsyncProgress&)>&& Function) override
+		{
+			OnProgressFn = MoveTemp(Function);
+		}
+
+		virtual void SetOnWillRetry(TDelegate<void(TOnlineAsyncOpHandle<OpType>& Handle, const FWillRetry&)>&& Function) override
+		{
+			OnWillRetryFn = MoveTemp(Function);
+		}
+
+		virtual void SetOnComplete(TDelegate<void(const TOnlineResult<ResultType>&)>&& Function) override
+		{
+			OnCompleteFn = MoveTemp(Function);
+			if (SharedState->IsComplete())
+			{
+				TriggerOnComplete(SharedState->Result);
+			}
+		}
+
+		void TriggerOnComplete(const TOnlineResult<ResultType>& Result)
+		{
+			// TODO: Execute OnCompleteFn next tick on game thread
+			if (OnCompleteFn.IsBound())
+			{
+				OnCompleteFn.ExecuteIfBound(Result);
+				OnCompleteFn.Unbind();
+				Detach();
+			}
+		}
+
+	private:
+		void Detach()
+		{
+			TSharedPtr<TOnlineAsyncOp<OpType>> PinnedOp = AsyncOp.Pin();
+			AsyncOp.Reset();
+			if (PinnedOp.IsValid())
+			{
+				PinnedOp->Detach(this->AsShared());
+			}
+		}
+
+		TDelegate<void(const FAsyncProgress&)> OnProgressFn;
+		TDelegate<void(TOnlineAsyncOpHandle<OpType>& Handle, const FWillRetry&)> OnWillRetryFn;
+		TDelegate<void(const TOnlineResult<ResultType>&)> OnCompleteFn;
+
+		bool bCancelled = false;
+		TSharedRef<FAsyncOpSharedState> SharedState;
+		TWeakPtr<TOnlineAsyncOp<OpType>> AsyncOp;
+	};
+
+	void Detach(const TSharedRef<FAsyncOpSharedHandleState>& SharedHandleState)
+	{
+		SharedHandleStates.Remove(SharedHandleState);
+	}
+
+	TSharedRef<Private::IOnlineAsyncOpSharedState<OpType>> CreateSharedState()
+	{
+		TSharedRef<FAsyncOpSharedHandleState> SharedHandleState = MakeShared<FAsyncOpSharedHandleState>(this->AsShared());
+		SharedHandleStates.Add(SharedHandleState);
+		return StaticCastSharedRef<Private::IOnlineAsyncOpSharedState<OpType>>(SharedHandleState);
+	}
+
+	TSharedRef<FAsyncOpSharedState> SharedState;
+	TArray<TSharedRef<FAsyncOpSharedHandleState>> SharedHandleStates;
+	TPromise<void> StartExecutionPromise;
+};
+
+/* UE::Online */ }
diff --git a/Engine/Plugins/Online/Experimental/OnlineServices/Source/OnlineServicesCommon/Public/Online/OnlineAsyncOpCache.h b/Engine/Plugins/Online/Experimental/OnlineServices/Source/OnlineServicesCommon/Public/Online/OnlineAsyncOpCache.h
new file mode 100644
--- /dev/null
+++ b/Engine/Plugins/Online/Experimental/OnlineServices/Source/OnlineServicesCommon/Public/Online/OnlineAsyncOpCache.h
@@ -0,0 +1,283 @@
+// Copyright Epic Games, Inc. All Rights Reserved.
+
+#pragma once
+
+#include "Online/OnlineAsyncOp.h"
+#include "Online/OnlineId.h"
+#include "Containers/Map.h"
+#include "Templates/SharedPointer.h"
+
+namespace UE::Online {
+
+struct CLocalUserIdDefined
+{
+	template <typename T>
+	auto Requires() -> decltype(T::LocalUserId);
+};
+
+struct FOperationConfig
+{
+	float CacheExpirySeconds = 0.0f;
+};
+
+class IOnlineAnyData
+{
+public:
+	virtual ~IOnlineAnyData() {}
+	virtual FOnlineTypeName GetTypeName() const = 0;
+
+	template <typename T>
+	const T* Get() const
+	{
+		if (GetTypeName() == TOnlineTypeInfo<T>::GetTypeName())
+		{
+			return static_cast<const T*>(GetData());
+		}
+
+		return nullptr;
+	}
+
+	template <typename T>
+	const T& GetRef() const
+	{
+		const T* Value = Get<T>();
+		check(Value != nullptr);
+		return *Value;
+	}
+
+protected:
+	virtual const void* GetData() const = 0;
+};
+
+template <typename T, typename BaseType = IOnlineAnyData>
+class TOnlineAnyData : public BaseType
+{
+public:
+	using ValueType = std::remove_reference_t<T>;
+
+	TOnlineAnyData(const ValueType& InData)
+		: Data(InData)
+	{
+	}
+
+	TOnlineAnyData(ValueType&& InData)
+		: Data(MoveTemp(InData))
+	{
+	}
+
+	virtual FOnlineTypeName GetTypeName() const override
+	{
+		return  TOnlineTypeInfo<ValueType>::GetTypeName();
+	}
+
+	const T& GetDataRef() const 
+	{
+		return Data;
+	}
+
+protected:
+	virtual const void* GetData() const override
+	{
+		return &Data;
+	}
+
+	T Data;
+};
+
+template <typename OpType>
+struct TAsyncOpParamsFuncs
+{
+	inline static bool Compare(const typename OpType::Params& First, const typename OpType::Params& Second)
+	{
+		bool bResult = true;
+		Meta::VisitFields<typename OpType::Params>([&bResult, &First, &Second](const auto& Field)
+		{
+			bResult = bResult && (First.*Field.Pointer) != (Second.*Field.Pointer);
+		});
+		return bResult;
+	}
+
+	inline static uint32 GetTypeHash(const typename OpType::Params& Params)
+	{
+		using ::GetTypeHash;
+		uint32 CombinedHash = 0;
+		Meta::VisitFields(Params,
+			[&CombinedHash](const TCHAR* FieldName, const auto& Field)
+			{
+				//HashCombine(CombinedHash, GetTypeHash(Field));
+			});
+		return CombinedHash;
+	}
+};
+
+class FOnlineAsyncOpCache
+{
+public:
+	FOnlineAsyncOpCache(const FString& ConfigPath, class FOnlineServicesCommon& InServices)
+		: Services(InServices)
+	{
+	}
+
+	// Create an operation
+	template <typename OpType>
+	TOnlineAsyncOp<OpType>& GetOp(typename OpType::Params&& Params)
+	{
+		TSharedRef<TOnlineAsyncOp<OpType>> Op = CreateOp<OpType>(MoveTemp(Params));
+
+		return *Op;
+	}
+
+	// Join an existing operation or use a non-expired cached result, or create an operation that can later be joined
+	template <typename OpType>
+	TOnlineAsyncOp<OpType>& GetJoinableOp(typename OpType::Params&& Params)
+	{
+		TSharedPtr<TOnlineAsyncOp<OpType>> Op;
+
+		// find existing op
+		if (false)
+		{
+
+		}
+		else
+		{
+			Op = CreateOp<OpType>(MoveTemp(Params));
+		}
+
+		return *Op;
+	}
+
+	// Merge with a pending operation, or create an operation
+	template <typename OpType>
+	TOnlineAsyncOp<OpType>& GetMergableOp(typename OpType::Params&& Params)
+	{
+		TSharedPtr<TOnlineAsyncOp<OpType>> Op;
+
+		// find existing op
+		if (false)
+		{
+
+		}
+		else
+		{
+			Op = CreateOp<OpType>(MoveTemp(Params));
+		}
+
+		return *Op;
+	}
+
+	FOnlineServicesCommon& Services;
+
+private:
+	template <typename OpType>
+	TSharedRef<TOnlineAsyncOp<OpType>> CreateOp(typename OpType::Params&& Params)
+	{
+		TUniquePtr<TWrappedOperation<OpType>> WrappedOp = MakeUnique<TWrappedOperation<OpType>>(Services, MoveTemp(Params));
+
+		TSharedRef<TOnlineAsyncOp<OpType>> Op = WrappedOp->GetDataRef();
+
+		if constexpr (TModels<CLocalUserIdDefined, typename OpType::Params>::Value)
+		{
+			// add LocalUserId to the Op Data
+			Op->Data.template Set<decltype(Params.LocalUserId)>(TEXT("LocalUserId"), Op->GetParams().LocalUserId);
+		}
+
+		// TODO: This needs to be aware of operation cache expiry policy
+		Operations.Add(FWrappedOperationKey(*Op), MoveTemp(WrappedOp));
+
+		return MoveTemp(Op);
+	}
+
+	class IWrappedOperation : public IOnlineAnyData
+	{
+	public:
+		virtual bool IsJoinable() = 0;
+		virtual bool IsMergable() = 0;
+	};
+
+	template <typename OpType>
+	class TWrappedOperation : public TOnlineAnyData<TSharedRef<TOnlineAsyncOp<OpType>>, IWrappedOperation>
+	{
+	public:
+		template <typename... ParamTypes>
+		TWrappedOperation(ParamTypes&&... Params)
+			: TOnlineAnyData<TSharedRef<TOnlineAsyncOp<OpType>>, IWrappedOperation>(MakeShared<TOnlineAsyncOp<OpType>>(Forward<ParamTypes>(Params)...))
+		{
+		}
+
+		virtual bool IsJoinable() override { return true; }
+		virtual bool IsMergable() override { return false; }
+	};
+
+	class FWrappedOperationKey
+	{
+	public:
+		template <typename OpType>
+		FWrappedOperationKey(const TOnlineAsyncOp<OpType>& Operation)
+			: Impl(new TWrappedOperationKeyImpl<OpType>(Operation.GetParams()))
+		{
+		}
+
+		bool operator==(const FWrappedOperationKey& Other) const
+		{
+			return Impl->Compare(*Other.Impl);
+		}
+
+		uint32 GetTypeHash() const
+		{
+			return Impl->GetTypeHash();
+		}
+
+	private:
+		class IWrappedOperationKeyImpl : public IOnlineAnyData
+		{
+		public:
+			virtual bool Compare(const IWrappedOperationKeyImpl& Other) const = 0;
+			virtual uint32 GetTypeHash() const = 0;
+		};
+
+		template <typename OpType>
+		class TWrappedOperationKeyImpl : public TOnlineAnyData<const typename OpType::Params&, IWrappedOperationKeyImpl>
+		{
+		public:
+			using DataType = const typename OpType::Params&;
+			using ValueType = const typename OpType::Params;
+
+			TWrappedOperationKeyImpl(DataType& ParamsRef)
+				: TOnlineAnyData<DataType&, IWrappedOperationKeyImpl>(ParamsRef)
+			{
+			}
+
+			virtual bool Compare(const IWrappedOperationKeyImpl& Other) const override
+			{
+				if (Other.GetTypeName() == this->GetTypeName())
+				{
+					return TAsyncOpParamsFuncs<OpType>::Compare(this->template GetRef<ValueType>(), Other.template GetRef<ValueType>());
+				}
+
+				return false;
+			}
+
+			virtual uint32 GetTypeHash() const override
+			{
+				const uint32 Hash = TAsyncOpParamsFuncs<OpType>::GetTypeHash(this->template GetRef<ValueType>());
+				return HashCombine(::UE::Online::GetTypeHash(this->GetTypeName()), Hash);
+			}
+		};
+
+		TUniquePtr<IWrappedOperationKeyImpl> Impl;
+	};
+
+	friend uint32 GetTypeHash(const FWrappedOperationKey& Key);
+
+	TMap<FWrappedOperationKey, TUniquePtr<IWrappedOperation>> Operations;
+	TMap<FAccountId, TMap<FWrappedOperationKey, TUniquePtr<IWrappedOperation>>> UserOperations;
+};
+
+inline uint32 GetTypeHash(const FOnlineAsyncOpCache::FWrappedOperationKey& Key)
+{
+	return Key.GetTypeHash();
+}
+
+/* UE::Online */ }
+
+#include "OnlineAsyncOpCache_Meta.inl"
diff --git a/Engine/Plugins/Online/Experimental/OnlineServices/Source/OnlineServicesCommon/Public/Online/OnlineAsyncOpCache_Meta.inl b/Engine/Plugins/Online/Experimental/OnlineServices/Source/OnlineServicesCommon/Public/Online/OnlineAsyncOpCache_Meta.inl
new file mode 100644
--- /dev/null
+++ b/Engine/Plugins/Online/Experimental/OnlineServices/Source/OnlineServicesCommon/Public/Online/OnlineAsyncOpCache_Meta.inl
@@ -0,0 +1,13 @@
+// Copyright Epic Games, Inc. All Rights Reserved.
+
+#pragma once
+
+#include "Online/OnlineMeta.h"
+
+namespace UE::Online::Meta {
+
+BEGIN_ONLINE_STRUCT_META(FOperationConfig)
+	ONLINE_STRUCT_FIELD(FOperationConfig, CacheExpirySeconds)
+END_ONLINE_STRUCT_META()
+
+/* UE::Online::Meta*/ }
\ No newline at end of file
diff --git a/Engine/Plugins/Online/Experimental/OnlineServices/Source/OnlineServicesCommon/Public/Online/OnlineComponent.h b/Engine/Plugins/Online/Experimental/OnlineServices/Source/OnlineServicesCommon/Public/Online/OnlineComponent.h
new file mode 100644
--- /dev/null
+++ b/Engine/Plugins/Online/Experimental/OnlineServices/Source/OnlineServicesCommon/Public/Online/OnlineComponent.h
@@ -0,0 +1,92 @@
+// Copyright Epic Games, Inc. All Rights Reserved.
+
+#pragma once
+
+#include "Templates/SharedPointer.h"
+
+namespace UE::Online {
+
+class IOnlineComponent
+{
+public:
+	virtual ~IOnlineComponent() {}
+	// Called after component has been constructed. It is not safe to reference other components at this time
+	virtual void Initialize() = 0;
+	// Called after all components have been initialized
+	virtual void PostInitialize() = 0;
+	// Called whenever we need to reload data from config
+	virtual void LoadConfig() = 0;
+	// Called every Tick
+	virtual void Tick(float DeltaSeconds) = 0;
+	// Called before any component has been shutdown
+	virtual void PreShutdown() = 0;
+	// Called right before the component is destroyed. It is not safe to reference any other components at this time
+	virtual void Shutdown() = 0;
+};
+
+template <typename ComponentType>
+class TOnlineComponent
+	: public ComponentType
+	, public IOnlineComponent
+{
+public:
+	TOnlineComponent(const TOnlineComponent&) = delete;
+	TOnlineComponent(TOnlineComponent&&) = delete;
+
+	/**
+	 */
+	template <typename U>
+	TOnlineComponent(const FString& ComponentName, const TSharedRef<U>& OwningRef)
+		: WeakThis(TSharedPtr<ComponentType>(OwningRef, static_cast<ComponentType*>(this)))
+		{
+		}
+
+	virtual void Initialize() override {}
+	virtual void PostInitialize() override {}
+	virtual void LoadConfig() override {}
+	virtual void Tick(float DeltaSeconds) override {}
+	virtual void PreShutdown() override {}
+	virtual void Shutdown() override {}
+
+	template <typename StructType>
+	bool LoadConfig(StructType& Struct, const FString& OperationName = FString()) const
+	{
+		return static_cast<const ComponentType*>(this)->GetServices()->LoadConfig(Struct, GetConfigSectionName(), OperationName);
+	}
+
+	const FString& GetConfigSectionName() const { return ConfigSectionName; }
+
+	// TSharedFromThis-like behaviour
+	TSharedRef<ComponentType> AsShared()
+	{
+		TSharedPtr<ComponentType> SharedThis(WeakThis.Pin());
+		check(SharedThis.Get() == this);
+		return MoveTemp(SharedThis).ToSharedRef();
+	}
+
+	TSharedRef<ComponentType const> AsShared() const
+	{
+		TSharedPtr<ComponentType const> SharedThis(WeakThis.Pin());
+		check(SharedThis.Get() == this);
+		return MoveTemp(SharedThis).ToSharedRef();
+	}
+
+protected:
+	template <class OtherType>
+	static TSharedRef<OtherType> SharedThis(OtherType* ThisPtr)
+	{
+		return StaticCastSharedRef<OtherType>(ThisPtr->AsShared());
+	}
+
+	template <class OtherType>
+	static TSharedRef<OtherType const> SharedThis(const OtherType* ThisPtr)
+	{
+		return StaticCastSharedRef<OtherType const>(ThisPtr->AsShared());
+	}
+
+private:
+	TWeakPtr<ComponentType> WeakThis;
+	FString ConfigSectionName;
+};
+
+/* UE::Online */ }
diff --git a/Engine/Plugins/Online/Experimental/OnlineServices/Source/OnlineServicesCommon/Public/Online/OnlineComponentRegistry.h b/Engine/Plugins/Online/Experimental/OnlineServices/Source/OnlineServicesCommon/Public/Online/OnlineComponentRegistry.h
new file mode 100644
--- /dev/null
+++ b/Engine/Plugins/Online/Experimental/OnlineServices/Source/OnlineServicesCommon/Public/Online/OnlineComponentRegistry.h
@@ -0,0 +1,118 @@
+// Copyright Epic Games, Inc. All Rights Reserved.
+
+#pragma once
+
+#include <Containers/Array.h>
+#include <Templates/UniquePtr.h>
+
+#include "Online/OnlineComponent.h"
+#include "Online/OnlineTypeInfo.h"
+#include "Online/OnlineMeta.h"
+
+namespace UE::Online {
+
+class FOnlineComponentRegistry
+{
+public:
+	/**
+	 *
+	 */
+	template <typename T, typename... ParamTypes>
+	void Register(ParamTypes&&... Params)
+	{
+		TUniquePtr<IComponentWrapper> Component = MakeUnique<TComponentWrapper<T>>(Forward<ParamTypes>(Params)...);
+
+		Components.Add(MoveTemp(Component));
+	}
+
+	/**
+	 *
+	 */
+	template <typename ComponentType>
+	ComponentType* Get() const
+	{
+		const FOnlineTypeName TypeName = TOnlineTypeInfo<ComponentType>::GetTypeName();
+		const TUniquePtr<IComponentWrapper>* WrappedComponent = Components.FindByPredicate([&TypeName](const TUniquePtr<IComponentWrapper>& WrappedComponent) { return WrappedComponent->IsA(TypeName); });
+		if (WrappedComponent)
+		{
+			typename Meta::TBaseClass<ComponentType>::Type* BasePtr = static_cast<typename Meta::TBaseClass<ComponentType>::Type*>((*WrappedComponent)->GetBase());
+			return static_cast<ComponentType*>(BasePtr);
+		}
+		return nullptr;
+	}
+
+	/**
+	 *
+	 */
+	template <typename CallableType, typename... ParamTypes>
+	void Visit(CallableType&& Callable, ParamTypes&&... Params)
+	{
+		for (auto&& Component : Components)
+		{
+			Invoke(Callable, *Component, Forward<ParamTypes>(Params)...);
+		}
+	}
+
+private:
+	class IComponentWrapper
+	{
+	public:
+		virtual ~IComponentWrapper() {}
+		virtual void* GetBase() = 0;
+		virtual bool IsA(FOnlineTypeName TypeName) const = 0;
+		virtual FOnlineTypeName GetTypeName() const = 0;
+		virtual IOnlineComponent* Get() = 0;
+		virtual const IOnlineComponent* Get() const = 0;
+		IOnlineComponent& operator*() { return *Get(); }
+		const IOnlineComponent& operator*() const { return *Get(); }
+		IOnlineComponent* operator->() { return Get(); }
+		const IOnlineComponent* operator->() const { return Get(); }
+	};
+
+	template <typename T>
+	class TComponentWrapper : public IComponentWrapper
+	{
+	public:
+		template <typename... TArgs>
+		TComponentWrapper(TArgs&&... Args)
+			: Object(Forward<TArgs>(Args)...)
+		{
+		}
+
+		virtual ~TComponentWrapper()
+		{
+		}
+
+		virtual void* GetBase() override final
+		{
+			return static_cast<void*>(static_cast<Meta::TBaseClass_T<T>*>(&Object));
+		}
+
+		virtual bool IsA(FOnlineTypeName TypeName) const override final
+		{
+			return TOnlineTypeInfo<T>::IsA(TypeName);
+		}
+
+		virtual FOnlineTypeName GetTypeName() const override final
+		{
+			return TOnlineTypeInfo<T>::GetTypeName();
+		}
+
+		virtual IOnlineComponent* Get() override final
+		{
+			return &Object;
+		}
+
+		virtual const IOnlineComponent* Get() const override final
+		{
+			return &Object;
+		}
+
+	private:
+		T Object;
+	};
+
+	TArray<TUniquePtr<IComponentWrapper>> Components;
+};
+
+/* UE::Online */ }
diff --git a/Engine/Plugins/Online/Experimental/OnlineServices/Source/OnlineServicesCommon/Public/Online/OnlineConfig.h b/Engine/Plugins/Online/Experimental/OnlineServices/Source/OnlineServicesCommon/Public/Online/OnlineConfig.h
new file mode 100644
--- /dev/null
+++ b/Engine/Plugins/Online/Experimental/OnlineServices/Source/OnlineServicesCommon/Public/Online/OnlineConfig.h
@@ -0,0 +1,407 @@
+// Copyright Epic Games, Inc. All Rights Reserved.
+
+#pragma once
+
+#include "Misc/ConfigCacheIni.h"
+#include "Online/OnlineMeta.h"
+
+namespace UE::Online {
+
+/**
+ * Interface for retrieving config values used by OnlineServices implementations
+ * 
+ * For all the GetValue functions, Section will correspond to the section in an ini file, but can be
+ * interpreted in other ways by an IOnlineConfigProvider implementation to allow for OnlineServices
+ * implementations to use an alternative configuration file.
+ * 
+ * By convention, all sections of the OnlineServices implementations will take the following format:
+ *   OnlineServices.<ServicesProvider>.<OptionalInterface>.<OptionalOperation> <OptionalOverride>
+ * For example, valid sections include the following:
+ *   OnlineServices.Null
+ *   OnlineServices.Null.Auth
+ *   OnlineServices.Null.Auth.Login
+ *   OnlineServices.Null.Auth.Login Prod
+ * 
+ * Implementations must at a minimum implement the FString and TArray<FString> GetValue. The others are
+ * optional, and the default implementation will convert strings to the appropriate data type, but can 
+ * be overridden in cases where the underlying configuration system stores integers, floats, etc directly
+ * instead of as strings
+ */
+class ONLINESERVICESCOMMON_API IOnlineConfigProvider
+{
+public:
+	virtual ~IOnlineConfigProvider() {}
+
+	/**
+	 * Get a FString value
+	 * 
+	 * @param Section Section to read the value from
+	 * @param Key Key in the section for the value
+	 * @param Value FString value read from the config. Will be unchanged if not present
+	 * 
+	 * @return true if a value was read
+	 */
+	virtual bool GetValue(const TCHAR* Section, const TCHAR* Key, FString& Value) = 0;
+
+	/**
+	 * Get an value consisting of a TArray of FStrings
+	 *
+	 * @param Section Section to read the value from
+	 * @param Key Key in the section for the value
+	 * @param Value Array of FStrings read from the config. Will be empty if not present
+	 *
+	 * @return number of values in the array
+	 */
+	virtual int32 GetValue(const TCHAR* Section, const TCHAR* Key, TArray<FString>& Value) = 0;
+
+	/**
+	 * Get an FText value
+	 *
+	 * @param Section Section to read the value from
+	 * @param Key Key in the section for the value
+	 * @param Value FText value read from the config. Will be unchanged if not present
+	 *
+	 * @return true if a value was read
+	 */
+	virtual bool GetValue(const TCHAR* Section, const TCHAR* Key, FText& Value)
+	{
+		FString StringValue;
+		if (GetValue(Section, Key, StringValue))
+		{
+			return FTextStringHelper::ReadFromBuffer(*StringValue, Value, Section) != nullptr;
+		}
+		return false;
+	}
+
+	/**
+	 * Get a bool value
+	 *
+	 * @param Section Section to read the value from
+	 * @param Key Key in the section for the value
+	 * @param Value bool value read from the config. Will be unchanged if not present
+	 *
+	 * @return true if a value was read
+	 */
+	virtual bool GetValue(const TCHAR* Section, const TCHAR* Key, bool& Value)
+	{
+		FString StringValue;
+		if (GetValue(Section, Key, StringValue))
+		{
+			Value = FCString::ToBool(*StringValue);
+			return true;
+		}
+		return false;
+	}
+
+	/**
+	 * Get an int32 value
+	 *
+	 * @param Section Section to read the value from
+	 * @param Key Key in the section for the value
+	 * @param Value int32 value read from the config. Will be unchanged if not present
+	 *
+	 * @return true if a value was read
+	 */
+	virtual bool GetValue(const TCHAR* Section, const TCHAR* Key, int32& Value)
+	{
+		FString StringValue;
+		if (GetValue(Section, Key, StringValue))
+		{
+			Value = FCString::Strtoi(*StringValue, nullptr, 10);
+			return true;
+		}
+		return false;
+	}
+
+	/**
+	 * Get an int64 value
+	 *
+	 * @param Section Section to read the value from
+	 * @param Key Key in the section for the value
+	 * @param Value int64 value read from the config. Will be unchanged if not present
+	 *
+	 * @return true if a value was read
+	 */
+	virtual bool GetValue(const TCHAR* Section, const TCHAR* Key, int64& Value)
+	{
+		FString StringValue;
+		if (GetValue(Section, Key, StringValue))
+		{
+			Value = FCString::Strtoi64(*StringValue, nullptr, 10);
+			return true;
+		}
+		return false;
+	}
+
+	/**
+	 * Get a uint64 value
+	 *
+	 * @param Section Section to read the value from
+	 * @param Key Key in the section for the value
+	 * @param Value uint64 value read from the config. Will be unchanged if not present
+	 *
+	 * @return true if a value was read
+	 */
+	virtual bool GetValue(const TCHAR* Section, const TCHAR* Key, uint64& Value)
+	{
+		FString StringValue;
+		if (GetValue(Section, Key, StringValue))
+		{
+			Value = FCString::Strtoui64(*StringValue, nullptr, 10);
+			return true;
+		}
+		return false;
+	}
+
+	/**
+	 * Get a float value
+	 *
+	 * @param Section Section to read the value from
+	 * @param Key Key in the section for the value
+	 * @param Value float value read from the config. Will be unchanged if not present
+	 *
+	 * @return true if a value was read
+	 */
+	virtual bool GetValue(const TCHAR* Section, const TCHAR* Key, float& Value)
+	{
+		FString StringValue;
+		if (GetValue(Section, Key, StringValue))
+		{
+			Value = FCString::Atof(*StringValue);
+			return true;
+		}
+		return false;
+	}
+
+	/**
+	 * Get a double value
+	 *
+	 * @param Section Section to read the value from
+	 * @param Key Key in the section for the value
+	 * @param Value double value read from the config. Will be unchanged if not present
+	 *
+	 * @return true if a value was read
+	 */
+	virtual bool GetValue(const TCHAR* Section, const TCHAR* Key, double& Value)
+	{
+		FString StringValue;
+		if (GetValue(Section, Key, StringValue))
+		{
+			Value = FCString::Atod(*StringValue);
+			return true;
+		}
+		return false;
+	}
+
+	/**
+	 * Get a value consisting of a TArray of bool
+	 *
+	 * @param Section Section to read the value from
+	 * @param Key Key in the section for the value
+	 * @param Value Array of bool read from the config. Will be empty if not present
+	 *
+	 * @return number of values in the array
+	 */
+	virtual int32 GetValue(const TCHAR* Section, const TCHAR* Key, TArray<bool>& Value)
+	{
+		Value.Empty();
+		TArray<FString> StringArray;
+		if (GetValue(Section, Key, StringArray) > 0)
+		{
+			Value.Reserve(StringArray.Num());
+			for (const FString& StringValue : StringArray)
+			{
+				Value.Emplace(FCString::ToBool(*StringValue));
+			}
+		}
+		return Value.Num();
+	}
+
+	/**
+	 * Get a value consisting of a TArray of int32
+	 *
+	 * @param Section Section to read the value from
+	 * @param Key Key in the section for the value
+	 * @param Value Array of int32 read from the config. Will be empty if not present
+	 *
+	 * @return number of values in the array
+	 */
+	virtual int32 GetValue(const TCHAR* Section, const TCHAR* Key, TArray<int32>& Value)
+	{
+		Value.Empty();
+		TArray<FString> StringArray;
+		if (GetValue(Section, Key, StringArray) > 0)
+		{
+			Value.Reserve(StringArray.Num());
+			for (const FString& StringValue : StringArray)
+			{
+				Value.Emplace(FCString::Strtoi(*StringValue, nullptr, 10));
+			}
+		}
+		return Value.Num();
+	}
+
+	/**
+	 * Get a value consisting of a TArray of int64
+	 *
+	 * @param Section Section to read the value from
+	 * @param Key Key in the section for the value
+	 * @param Value Array of int64 read from the config. Will be empty if not present
+	 *
+	 * @return number of values in the array
+	 */
+	virtual int32 GetValue(const TCHAR* Section, const TCHAR* Key, TArray<int64>& Value)
+	{
+		Value.Empty();
+		TArray<FString> StringArray;
+		if (GetValue(Section, Key, StringArray) > 0)
+		{
+			Value.Reserve(StringArray.Num());
+			for (const FString& StringValue : StringArray)
+			{
+				Value.Emplace(FCString::Strtoi64(*StringValue, nullptr, 10));
+			}
+		}
+		return Value.Num();
+	}
+
+	/**
+	 * Get a value consisting of a TArray of uint64
+	 *
+	 * @param Section Section to read the value from
+	 * @param Key Key in the section for the value
+	 * @param Value Array of uint64 read from the config. Will be empty if not present
+	 *
+	 * @return number of values in the array
+	 */
+	virtual int32 GetValue(const TCHAR* Section, const TCHAR* Key, TArray<uint64>& Value)
+	{
+		Value.Empty();
+		TArray<FString> StringArray;
+		if (GetValue(Section, Key, StringArray) > 0)
+		{
+			Value.Reserve(StringArray.Num());
+			for (const FString& StringValue : StringArray)
+			{
+				Value.Emplace(FCString::Strtoui64(*StringValue, nullptr, 10));
+			}
+		}
+		return Value.Num();
+	}
+
+	/**
+	 * Get a value consisting of a TArray of float
+	 *
+	 * @param Section Section to read the value from
+	 * @param Key Key in the section for the value
+	 * @param Value Array of float read from the config. Will be empty if not present
+	 *
+	 * @return number of values in the array
+	 */
+	virtual int32 GetValue(const TCHAR* Section, const TCHAR* Key, TArray<float>& Value)
+	{
+		Value.Empty();
+		TArray<FString> StringArray;
+		if (GetValue(Section, Key, StringArray) > 0)
+		{
+			Value.Reserve(StringArray.Num());
+			for (const FString& StringValue : StringArray)
+			{
+				Value.Emplace(FCString::Atof(*StringValue));
+			}
+		}
+		return Value.Num();
+	}
+
+	/**
+	 * Get a value consisting of a TArray of double
+	 *
+	 * @param Section Section to read the value from
+	 * @param Key Key in the section for the value
+	 * @param Value Array of double read from the config. Will be empty if not present
+	 *
+	 * @return number of values in the array
+	 */
+	virtual int32 GetValue(const TCHAR* Section, const TCHAR* Key, TArray<double>& Value)
+	{
+		Value.Empty();
+		TArray<FString> StringArray;
+		if (GetValue(Section, Key, StringArray) > 0)
+		{
+			Value.Reserve(StringArray.Num());
+			for (const FString& StringValue : StringArray)
+			{
+				Value.Emplace(FCString::Atod(*StringValue));
+			}
+		}
+		return Value.Num();
+	}
+};
+
+/**
+ * IOnlineConfigProvider implementation that uses GConfig to retrieve values
+ */
+class ONLINESERVICESCOMMON_API FOnlineConfigProviderGConfig : public IOnlineConfigProvider
+{
+public:
+	FOnlineConfigProviderGConfig(const FString& InConfigFile)
+		: ConfigFile(InConfigFile)
+	{
+	}
+
+	virtual bool GetValue(const TCHAR* Section, const TCHAR* Key, FString& Value) override
+	{
+		return GConfig->GetValue(Section, Key, Value, ConfigFile);
+	}
+
+	virtual int32 GetValue(const TCHAR* Section, const TCHAR* Key, TArray<FString>& Value) override
+	{
+		return GConfig->GetValue(Section, Key, Value, ConfigFile);
+	}
+
+private:
+	FString ConfigFile;
+};
+
+namespace Private {
+
+// If needed, this can be specialized for additional types generically using LexFromString or explicit specialization
+template <typename T>
+auto LoadConfigValue(IOnlineConfigProvider& Provider, const TCHAR* Section, const TCHAR* Key, T& Value)
+	-> std::enable_if_t<std::is_same_v<decltype(Provider.GetValue(Section, Key, Value)), bool>, bool>
+{
+	return Provider.GetValue(Section, Key, Value);
+}
+
+template <typename T>
+auto LoadConfigValue(IOnlineConfigProvider& Provider, const TCHAR* Section, const TCHAR* Key, T& Value)
+	-> std::enable_if_t<std::is_same_v<decltype(Provider.GetValue(Section, Key, Value)), int32>, bool>
+{
+	return Provider.GetValue(Section, Key, Value) > 0;
+}
+
+/* Private */ }
+
+/**
+ * Populate a struct from a config provider. It will load from keys matching the struct member names in the section
+ * Requires that the type metadata is specified for StructType
+ * 
+ * @param Provider The config provider
+ * @param Section The section of the config to load the values from
+ * @param The struct whose fields will be loaded from config
+ * 
+ * @return true if any values were loaded
+ */
+template <typename StructType>
+bool LoadConfig(IOnlineConfigProvider& Provider, const FString& Section, StructType& Struct)
+{
+	bool bLoadedValue = false;
+	Meta::VisitFields(Struct,
+		[&Provider, &Section, &bLoadedValue](const TCHAR* FieldName, auto& Field)
+		{
+			bLoadedValue |= Private::LoadConfigValue(Provider, *Section, FieldName, Field);
+		});
+	return bLoadedValue;
+}
+
+/* UE::Online */ }
diff --git a/Engine/Plugins/Online/Experimental/OnlineServices/Source/OnlineServicesCommon/Public/Online/OnlineMeta.h b/Engine/Plugins/Online/Experimental/OnlineServices/Source/OnlineServicesCommon/Public/Online/OnlineMeta.h
new file mode 100644
--- /dev/null
+++ b/Engine/Plugins/Online/Experimental/OnlineServices/Source/OnlineServicesCommon/Public/Online/OnlineMeta.h
@@ -0,0 +1,139 @@
+// Copyright Epic Games, Inc. All Rights Reserved.
+
+#pragma once
+
+#include "Templates/Models.h"
+#include "Templates/Tuple.h"
+
+namespace UE::Online::Meta {
+
+/**
+ * Concept that checks for the existence of a Super typedef/using declaration
+ */
+struct CSuperDefined
+{
+	template <typename T>
+	auto Requires() -> typename T::Super&;
+};
+
+namespace Private {
+
+template <typename T, bool THasSuper = TModels<CSuperDefined, T>::Value>
+struct TBaseClassHelper
+{
+};
+
+template <typename T>
+struct TBaseClassHelper<T, false>
+{
+	using Type = T;
+};
+
+template <typename T>
+struct TBaseClassHelper<T, true>
+{
+	using Type = typename TBaseClassHelper<typename T::Super>::Type;
+};
+
+template <typename StructType, typename FieldType>
+struct TField
+{
+	constexpr TField(FieldType StructType::* InPointer, const TCHAR* InName) : Pointer(InPointer), Name(InName) {}
+	using Type = FieldType;
+	Type StructType::* Pointer;
+	const TCHAR* Name;
+};
+
+template <typename StructType, typename FieldType>
+constexpr TField<StructType, FieldType> MakeTField(FieldType StructType::* Pointer, const TCHAR* Name)
+{
+	return TField<StructType, FieldType>(Pointer, Name);
+}
+
+/* Private*/ }
+
+/**
+ * If all types in a class hierarchy have a Super typedef or using declaration, Type will be the type of the last Super definition
+ */
+template <typename T>
+struct TBaseClass
+{
+	using Type = typename Private::TBaseClassHelper<T>::Type;
+};
+
+/**
+ * Convenience template alias for TBaseClass<T>::Type
+ */
+template <typename T>
+using TBaseClass_T = typename TBaseClass<T>::Type;
+
+/**
+ * Struct that needs to be specialized to provide the Fields in the struct. Defined using the ONLINE_STRUCT macros below
+ */
+template <typename StructType>
+struct TStructDetails
+{
+	static_assert(!std::is_same_v<StructType, StructType>, "Specialization of TStructDetails unavailable. Type must have metadata specified via the ONLINE_STRUCT macros");
+};
+
+/**
+ * Struct field definitions, manually defined via the ONLINE_STRUCT macros below
+ */
+template <typename StructType>
+struct TStruct : public TStructDetails<std::remove_cv_t<std::remove_reference_t<StructType>>>
+{
+};
+
+/**
+ * Visit the TFields for a struct. TStructDetails must be specialized for the StructType
+ * 
+ * @param Func The functor to apply.
+ */
+template <typename StructType, typename FuncType>
+inline void VisitFields(FuncType&& Func)
+{
+	VisitTupleElements(Forward<FuncType>(Func), TStruct<StructType>::Fields());
+}
+
+/**
+ * Visit the fields of a struct. TStructDetails must be specialized for the StructType
+ *
+ * @param Object The object whose fields will be applied to the functor
+ * @param Func The functor to apply.
+ */
+template <typename StructType, typename FuncType>
+inline void VisitFields(StructType&& Object, FuncType&& Func)
+{
+	VisitFields<StructType>(
+		[&Object, &Func](const auto& Field)
+		{
+			Func(Field.Name, Object.*Field.Pointer);
+		});
+}
+
+/**
+ * Visit the TFields of a struct. TStructDetails must be specialized for the StructType
+ *
+ * @param Object The object whose fields will be applied to the functor
+ * @param Func The functor to apply.
+ */
+template <typename StructType, typename FuncType>
+inline void VisitFields(const StructType& Object, FuncType&& Func)
+{
+	VisitFields<StructType>(
+		[&Object, &Func](const auto& Field)
+		{
+			Func(Field.Name, Object.*Field.Pointer);
+		});
+}
+
+/* UE::Online::Meta*/ }
+
+/**
+ * The following macros are used to define the metadata for Online structs
+ * Must be inside of the UE::Online::Meta namespace
+ */
+#define BEGIN_ONLINE_STRUCT_META(StructType) template <> struct TStructDetails<StructType> { static const auto& Fields() { static auto Fields = MakeTuple(
+#define ONLINE_STRUCT_FIELD(StructType, FieldName) Private::MakeTField(&StructType::FieldName, TEXT(#FieldName))
+#define ONLINE_STRUCT_FIELD_NAMED(StructType, FieldName, FieldNameString) Private::MakeTField(&StructType::FieldName, FieldNameString)
+#define END_ONLINE_STRUCT_META() ); return Fields; } };
diff --git a/Engine/Plugins/Online/Experimental/OnlineServices/Source/OnlineServicesCommon/Public/Online/OnlineServicesCommon.h b/Engine/Plugins/Online/Experimental/OnlineServices/Source/OnlineServicesCommon/Public/Online/OnlineServicesCommon.h
new file mode 100644
--- /dev/null
+++ b/Engine/Plugins/Online/Experimental/OnlineServices/Source/OnlineServicesCommon/Public/Online/OnlineServicesCommon.h
@@ -0,0 +1,211 @@
+// Copyright Epic Games, Inc. All Rights Reserved.
+
+#pragma once
+
+#include "Online/OnlineServices.h"
+
+#include "Online/OnlineComponentRegistry.h"
+#include "Online/OnlineAsyncOpCache.h"
+#include "Online/OnlineConfig.h"
+
+#include "Containers/Ticker.h"
+#include "Templates/SharedPointer.h"
+
+#include "Async/Async.h"
+
+namespace UE::Online {
+
+class ONLINESERVICESCOMMON_API FOnlineServicesCommon
+	: public IOnlineServices
+	, public TSharedFromThis<FOnlineServicesCommon>
+	, public FTickerObjectBase
+{
+public:
+	using Super = IOnlineServices;
+
+	FOnlineServicesCommon();
+	FOnlineServicesCommon(const FOnlineServicesCommon&) = delete;
+	FOnlineServicesCommon(FOnlineServicesCommon&&) = delete;
+	virtual ~FOnlineServicesCommon() {}
+
+	// IOnlineServices
+	virtual void Init() override;
+	virtual void Destroy() override;
+	virtual IAuthPtr GetAuthInterface() override;
+	virtual IFriendsPtr GetFriendsInterface() override;
+
+	// FOnlineServicesCommon
+
+	/**
+	 * Retrieve any of the Interface IOnlineComponents
+	 */
+	template <typename ComponentType>
+	ComponentType* Get()
+	{
+		return Components.Get<ComponentType>();
+	}
+
+	/**
+	 * Called to register all the IOnlineComponents with the IOnlineService, called after this is constructed
+	 */
+	virtual void RegisterComponents();
+
+	/**
+	 * Calls Initialize on all the components, called after RegisterComponents
+	 */
+	virtual void Initialize();
+
+	/**
+	 * Calls PostInitialize on all the components, called after Initialize
+	 */
+	virtual void PostInitialize();
+
+	/**
+	 * Calls LoadConfig on all the components
+	 */
+	virtual void LoadConfig();
+
+	/**
+	 * Calls Tick on all the components
+	 */
+	virtual bool Tick(float DeltaSeconds) override;
+
+	/**
+	 * Calls PreShutdown on all the components, called prior to Shutdown
+	 */
+	virtual void PreShutdown();
+
+	/**
+	 * Calls Shutdown on all the components, called before this is destructed
+	 */
+	virtual void Shutdown();
+
+	/**
+	 * Call a callable according to a specified execution policy
+	 */
+	template <typename CallableType>
+	void Execute(FOnlineAsyncExecutionPolicy ExecutionPolicy, CallableType&& Callable)
+	{
+		if (ExecutionPolicy.GetExecutionPolicy() == EOnlineAsyncExecutionPolicy::RunOnGameThread)
+		{
+			ExecuteOnGameThread(MoveTempIfPossible(Callable));
+		}
+	}
+
+	/**
+	 * Call a callable on the game thread
+	 */
+	template <typename CallableType>
+	void ExecuteOnGameThread(CallableType&& Callable)
+	{
+		if (IsInGameThread())
+		{
+			Callable();
+		}
+		else
+		{
+			Async(EAsyncExecution::TaskGraphMainThread, MoveTemp(Callable));
+		}
+	}
+
+	/**
+	 * Override the default config provider (FOnlineConfigProviderGConfig(GEngineini))
+	 */
+	void SetConfigProvider(TUniquePtr<IOnlineConfigProvider>&& InConfigProvider)
+	{
+		ConfigProvider = MoveTemp(InConfigProvider);
+	}
+
+	/**
+	 * Clear the list of config overrides
+	 */
+	void ResetConfigSectionOverrides()
+	{
+		ConfigSectionOverrides.Reset();
+	}
+
+	/**
+	 * Add a config section override. These will be used in the order they are added
+	 */
+	void AddConfigSectionOverride(const FString& Override)
+	{
+		ConfigSectionOverrides.Add(Override);
+	}
+
+	/**
+	 * Get the config name for the Subsystem
+	 */
+	const FString& GetConfigName() const { return ConfigName; }
+
+	/**
+	 * Load a config struct for an interface + operation
+	 * Will load values from the following sections:
+	 *   OnlineServices
+	 *   OnlineServices.<ServiceProvider>
+	 *   OnlineServices.<ServiceProvider>.<InterfaceName> (if InterfaceName is set)
+	 *   OnlineServices.<ServiceProvider>.<InterfaceName>.<OperationName> (if OperationName is set)
+	 * 
+	 * @param Struct Struct to populate with values from config
+	 * @param InterfaceName Optional interface name to append to the config section name
+	 * @param OperationName Optional operation name to append to the config section name
+	 * 
+	 * @return true if a value was loaded
+	 */
+	template <typename StructType>
+	bool LoadConfig(StructType& Struct, const FString& InterfaceName = FString(), const FString& OperationName = FString())
+	{
+		TArray<FString> SectionHeiarchy;
+		FString SectionName = TEXT("OnlineServices");
+		SectionHeiarchy.Add(SectionName);
+		SectionName += TEXT(".") + GetConfigName();
+		SectionHeiarchy.Add(SectionName);
+		if (!InterfaceName.IsEmpty())
+		{
+			SectionName += TEXT(".") + InterfaceName;
+			SectionHeiarchy.Add(SectionName);
+			if (!InterfaceName.IsEmpty())
+			{
+				SectionName += TEXT(".") + OperationName;
+				SectionHeiarchy.Add(SectionName);
+			}
+		}
+		return LoadConfig(Struct, SectionHeiarchy);
+	}
+
+	/**
+	 * Load a config struct for a section heiarchy, also using the ConfigSectionOverrides
+	 *
+	 * @param Struct Struct to populate with values from config
+	 * @param SectionHeiarchy Array of config sections to load values from
+	 *
+	 * @return true if a value was loaded
+	 */
+	template <typename StructType>
+	bool LoadConfig(StructType& Struct, const TArray<FString>& SectionHeiarchy)
+	{
+		bool bLoadedConfig = false;
+		for (const FString& Section : SectionHeiarchy)
+		{
+			bLoadedConfig |= LoadConfig(*ConfigProvider, Section, Struct);
+			for (const FString& Override : ConfigSectionOverrides)
+			{
+				FString OverrideSection = Section + TEXT(" ") + Override;
+				bLoadedConfig |= LoadConfig(*ConfigProvider, OverrideSection, Struct);
+			}
+		}
+		return bLoadedConfig;
+	}
+
+	FOnlineAsyncOpCache OpCache;
+
+protected:
+	FOnlineComponentRegistry Components;
+	TUniquePtr<IOnlineConfigProvider> ConfigProvider;
+
+	/* Config section overrides */
+	TArray<FString> ConfigSectionOverrides;
+	FString ConfigName;
+};
+
+/* UE::Online */ }
+
diff --git a/Engine/Plugins/Online/Experimental/OnlineServices/Source/OnlineServicesCommon/Public/Online/OnlineTypeInfo.h b/Engine/Plugins/Online/Experimental/OnlineServices/Source/OnlineServicesCommon/Public/Online/OnlineTypeInfo.h
new file mode 100644
--- /dev/null
+++ b/Engine/Plugins/Online/Experimental/OnlineServices/Source/OnlineServicesCommon/Public/Online/OnlineTypeInfo.h
@@ -0,0 +1,57 @@
+// Copyright Epic Games, Inc. All Rights Reserved.
+
+#pragma once
+
+#include "Misc/CString.h"
+#include "Misc/GeneratedTypeName.h"
+#include "Templates/RemoveReference.h"
+
+#include "Online/OnlineMeta.h"
+
+namespace UE::Online {
+
+struct FOnlineTypeName
+{
+	const TCHAR* Name = nullptr;
+
+	bool operator==(const FOnlineTypeName& Other) const
+	{
+		return FCString::Strcmp(Name, Other.Name) == 0;
+	}
+
+	bool operator!=(const FOnlineTypeName& Other) const
+	{
+		return !operator==(Other);
+	}
+};
+
+inline uint32 GetTypeHash(const FOnlineTypeName& TypeName)
+{
+	return ::GetTypeHash(TypeName.Name);
+}
+
+template <typename T>
+struct TOnlineTypeInfo
+{
+	static inline FOnlineTypeName GetTypeName()
+	{
+		return FOnlineTypeName{ GetGeneratedTypeName<std::remove_reference_t<T>>() };
+	}
+
+	static inline bool IsA(FOnlineTypeName TypeName)
+	{
+		if (TypeName == GetTypeName())
+		{
+			return true;
+		}
+
+		if constexpr (TModels<Meta::CSuperDefined, T>::Value)
+		{
+			return TOnlineTypeInfo<typename T::Super>::IsA(TypeName);
+		}
+
+		return false;
+	}
+};
+
+/* UE::Online */ }
diff --git a/Engine/Plugins/Online/Experimental/OnlineServices/Source/OnlineServicesInterface/OnlineServicesInterface.Build.cs b/Engine/Plugins/Online/Experimental/OnlineServices/Source/OnlineServicesInterface/OnlineServicesInterface.Build.cs
new file mode 100644
--- /dev/null
+++ b/Engine/Plugins/Online/Experimental/OnlineServices/Source/OnlineServicesInterface/OnlineServicesInterface.Build.cs
@@ -0,0 +1,27 @@
+// Copyright Epic Games, Inc. All Rights Reserved.
+
+using UnrealBuildTool;
+
+public class OnlineServicesInterface : ModuleRules
+{
+	public OnlineServicesInterface(ReadOnlyTargetRules Target) : base(Target)
+    {
+		PCHUsage = ModuleRules.PCHUsageMode.UseExplicitOrSharedPCHs;
+
+		bUseUnity = false; // set false to catch non-unity compile issues while prototyping
+
+		PublicDependencyModuleNames.AddRange(
+			new string[] {
+			}
+		);
+
+		PublicIncludePaths.Add(ModuleDirectory);
+
+		// OnlineService cannot depend on Engine!
+		PrivateDependencyModuleNames.AddRange(
+			new string[] { 
+				"Core",
+			}
+		);
+	}
+}
diff --git a/Engine/Plugins/Online/Experimental/OnlineServices/Source/OnlineServicesInterface/Private/OnlineServices.cpp b/Engine/Plugins/Online/Experimental/OnlineServices/Source/OnlineServicesInterface/Private/OnlineServices.cpp
new file mode 100644
--- /dev/null
+++ b/Engine/Plugins/Online/Experimental/OnlineServices/Source/OnlineServicesInterface/Private/OnlineServices.cpp
@@ -0,0 +1,19 @@
+// Copyright Epic Games, Inc. All Rights Reserved.
+
+#include "Online/OnlineServices.h"
+
+#include "Online/OnlineServicesRegistry.h"
+
+namespace UE::Online {
+
+TSharedPtr<IOnlineServices> GetServices(EOnlineServices OnlineServices, FName InstanceName)
+{
+	return FOnlineServicesRegistry::Get().GetNamedServicesInstance(OnlineServices, InstanceName);
+}
+
+void DestroyServices(EOnlineServices OnlineServices, FName InstanceName)
+{
+	FOnlineServicesRegistry::Get().DestroyNamedServicesInstance(OnlineServices, InstanceName);
+}
+
+/* UE::Online */ }
diff --git a/Engine/Plugins/Online/Experimental/OnlineServices/Source/OnlineServicesInterface/Private/OnlineServicesInterfaceModule.cpp b/Engine/Plugins/Online/Experimental/OnlineServices/Source/OnlineServicesInterface/Private/OnlineServicesInterfaceModule.cpp
new file mode 100644
--- /dev/null
+++ b/Engine/Plugins/Online/Experimental/OnlineServices/Source/OnlineServicesInterface/Private/OnlineServicesInterfaceModule.cpp
@@ -0,0 +1,17 @@
+// Copyright Epic Games, Inc. All Rights Reserved.
+
+#include "CoreMinimal.h"
+#include "Modules/ModuleManager.h"
+#include "Online/OnlineServicesRegistry.h"
+
+class FOnlineServicesInterfaceModule : public IModuleInterface
+{
+public:
+private:
+	virtual void ShutdownModule() override
+	{
+		UE::Online::FOnlineServicesRegistry::TearDown();
+	}
+};
+
+IMPLEMENT_MODULE(FOnlineServicesInterfaceModule, OnlineServicesInterface);
diff --git a/Engine/Plugins/Online/Experimental/OnlineServices/Source/OnlineServicesInterface/Private/OnlineServicesRegistry.cpp b/Engine/Plugins/Online/Experimental/OnlineServices/Source/OnlineServicesInterface/Private/OnlineServicesRegistry.cpp
new file mode 100644
--- /dev/null
+++ b/Engine/Plugins/Online/Experimental/OnlineServices/Source/OnlineServicesInterface/Private/OnlineServicesRegistry.cpp
@@ -0,0 +1,82 @@
+// Copyright Epic Games, Inc. All Rights Reserved.
+
+#include "Online/OnlineServicesRegistry.h"
+
+#include "Misc/LazySingleton.h"
+
+namespace UE::Online {
+
+FOnlineServicesRegistry& FOnlineServicesRegistry::Get()
+{
+	return TLazySingleton<FOnlineServicesRegistry>::Get();
+}
+
+void FOnlineServicesRegistry::TearDown()
+{
+	return TLazySingleton<FOnlineServicesRegistry>::TearDown();
+}
+
+void FOnlineServicesRegistry::RegisterServicesFactory(EOnlineServices OnlineServices, TUniquePtr<IOnlineServicesFactory>&& Factory, int32 Priority)
+{
+	FFactoryAndPriority* ExistingFactoryAndPriority = ServicesFactories.Find(OnlineServices);
+	if (ExistingFactoryAndPriority == nullptr || ExistingFactoryAndPriority->Priority < Priority)
+	{
+		ServicesFactories.Add(OnlineServices, FFactoryAndPriority(MoveTemp(Factory), Priority));
+	}
+}
+
+void FOnlineServicesRegistry::UnRegisterServicesFactory(EOnlineServices OnlineServices, int32 Priority)
+{
+	FFactoryAndPriority* ExistingFactoryAndPriority = ServicesFactories.Find(OnlineServices);
+	if (ExistingFactoryAndPriority != nullptr && ExistingFactoryAndPriority->Priority == Priority)
+	{
+		ServicesFactories.Remove(OnlineServices);
+	}
+}
+
+TSharedPtr<IOnlineServices> FOnlineServicesRegistry::GetNamedServicesInstance(EOnlineServices OnlineServices, FName InstanceName)
+{
+	TSharedPtr<IOnlineServices> Services;
+
+	if (TSharedPtr<IOnlineServices>* ServicesPtr = NamedServiceInstances.FindOrAdd(OnlineServices).Find(InstanceName))
+	{
+		Services = *ServicesPtr;
+	}
+	else
+	{
+		Services = CreateServices(OnlineServices);
+		if (Services.IsValid())
+		{
+			NamedServiceInstances.FindOrAdd(OnlineServices).Add(InstanceName, Services);
+		}
+	}
+
+	return Services;
+}
+
+void FOnlineServicesRegistry::DestroyNamedServicesInstance(EOnlineServices OnlineServices, FName InstanceName)
+{
+	if (TSharedPtr<IOnlineServices>* ServicesPtr = NamedServiceInstances.FindOrAdd(OnlineServices).Find(InstanceName))
+	{
+		(*ServicesPtr)->Destroy();
+
+		NamedServiceInstances.FindOrAdd(OnlineServices).Remove(InstanceName);
+	}
+}
+
+
+TSharedPtr<IOnlineServices> FOnlineServicesRegistry::CreateServices(EOnlineServices OnlineServices)
+{
+	TSharedPtr<IOnlineServices> Services;
+
+	FFactoryAndPriority* FactoryAndPriority = ServicesFactories.Find(OnlineServices);
+	if (FactoryAndPriority != nullptr)
+	{
+		Services = FactoryAndPriority->Factory->Create();
+		Services->Init();
+	}
+
+	return Services;
+}
+
+/* UE::Online */ }
diff --git a/Engine/Plugins/Online/Experimental/OnlineServices/Source/OnlineServicesInterface/Public/Online/Auth.h b/Engine/Plugins/Online/Experimental/OnlineServices/Source/OnlineServicesInterface/Public/Online/Auth.h
new file mode 100644
--- /dev/null
+++ b/Engine/Plugins/Online/Experimental/OnlineServices/Source/OnlineServicesInterface/Public/Online/Auth.h
@@ -0,0 +1,171 @@
+// Copyright Epic Games, Inc. All Rights Reserved.
+
+#pragma once
+
+#include "Online/OnlineAsyncOpHandle.h"
+#include "Online/OnlineId.h"
+#include "Online/OnlineMeta.h"
+
+namespace UE::Online {
+
+class FOnlineError;
+
+class FAccountInfo
+{
+public:
+	int32 LocalUserNum;
+	FAccountId UserId;
+	// TODO: Other fields
+};
+
+struct FAuthLogin
+{
+	struct Params
+	{
+		int32 LocalUserNum;
+		FString CredentialsType;
+		FString CredentialsId;
+		FString CredentialsToken;
+		TArray<FString> Scopes;
+	};
+
+	struct Result
+	{
+		TSharedRef<FAccountInfo> AccountInfo;
+	};
+};
+
+struct FAuthLogout
+{
+	struct Params
+	{
+		FAccountId LocalUserId;
+		bool bDestroyAuth = false;
+	};
+
+	struct Result
+	{
+	};
+};
+
+struct FAuthGenerateAuth
+{
+	struct Params
+	{
+		FAccountId LocalUserId;
+		FString Type;
+		TArray<FString> Scopes;
+	};
+
+	struct Result
+	{
+	};
+};
+
+struct FGetAccountByLocalUserNum
+{
+	struct Params
+	{
+		int32 LocalUserNum;
+	};
+
+	struct Result
+	{
+		TSharedRef<FAccountInfo> AccountInfo;
+	};
+};
+
+struct FGetAccountByAccountId
+{
+	struct Params
+	{
+		FAccountId LocalUserId;
+	};
+
+	struct Result
+	{
+		TSharedRef<FAccountInfo> AccountInfo;
+	};
+};
+
+struct FLoginStatusChanged
+{
+};
+
+class IAuth
+{
+public:
+	/**
+	 *
+	 */
+	virtual TOnlineAsyncOpHandle<FAuthLogin> Login(FAuthLogin::Params&& Params) = 0;
+
+	/**
+	 *
+	 */
+	virtual TOnlineAsyncOpHandle<FAuthLogout> Logout(FAuthLogout::Params&& Params) = 0;
+
+	/**
+	 *
+	 */
+	virtual TOnlineAsyncOpHandle<FAuthGenerateAuth> GenerateAuth(FAuthGenerateAuth::Params&& Params) = 0;
+
+	/**
+	 * Get logged in user by local user num
+	 */ 
+	virtual TOnlineResult<FGetAccountByLocalUserNum::Result> GetAccountByLocalUserNum(FGetAccountByLocalUserNum::Params&& Params) = 0;
+
+	/**
+	 * Get logged in user by account id
+	 */
+	virtual TOnlineResult<FGetAccountByAccountId::Result> GetAccountByAccountId(FGetAccountByAccountId::Params&& Params) = 0;
+
+	/**
+	 * Event triggered when a local user's login status changes
+	 */
+	virtual TOnlineEvent<void(const FLoginStatusChanged&)> OnLoginStatusChanged() = 0;
+};
+
+
+namespace Meta {
+// TODO: Move to Auth_Meta.inl file?
+
+BEGIN_ONLINE_STRUCT_META(FAccountInfo)
+	ONLINE_STRUCT_FIELD(FAccountInfo, LocalUserNum),
+	ONLINE_STRUCT_FIELD(FAccountInfo, UserId)
+END_ONLINE_STRUCT_META()
+
+BEGIN_ONLINE_STRUCT_META(FAuthLogin::Params)
+	ONLINE_STRUCT_FIELD(FAuthLogin::Params, LocalUserNum),
+	ONLINE_STRUCT_FIELD(FAuthLogin::Params, CredentialsType),
+	ONLINE_STRUCT_FIELD(FAuthLogin::Params, CredentialsId),
+	ONLINE_STRUCT_FIELD(FAuthLogin::Params, CredentialsToken),
+	ONLINE_STRUCT_FIELD(FAuthLogin::Params, Scopes)
+END_ONLINE_STRUCT_META()
+
+BEGIN_ONLINE_STRUCT_META(FAuthLogin::Result)
+	ONLINE_STRUCT_FIELD(FAuthLogin::Result, AccountInfo)
+END_ONLINE_STRUCT_META()
+
+BEGIN_ONLINE_STRUCT_META(FAuthLogout::Params)
+	ONLINE_STRUCT_FIELD(FAuthLogout::Params, LocalUserId)
+END_ONLINE_STRUCT_META()
+
+BEGIN_ONLINE_STRUCT_META(FAuthLogout::Result)
+END_ONLINE_STRUCT_META()
+
+BEGIN_ONLINE_STRUCT_META(FAuthGenerateAuth::Params)
+	ONLINE_STRUCT_FIELD(FAuthGenerateAuth::Params, LocalUserId),
+	ONLINE_STRUCT_FIELD(FAuthGenerateAuth::Params, Type),
+	ONLINE_STRUCT_FIELD(FAuthGenerateAuth::Params, Scopes)
+END_ONLINE_STRUCT_META()
+
+BEGIN_ONLINE_STRUCT_META(FAuthGenerateAuth::Result)
+END_ONLINE_STRUCT_META()
+
+BEGIN_ONLINE_STRUCT_META(FLoginStatusChanged)
+END_ONLINE_STRUCT_META()
+
+/* Meta*/ }
+
+/* UE::Online */ }
diff --git a/Engine/Plugins/Online/Experimental/OnlineServices/Source/OnlineServicesInterface/Public/Online/Friends.h b/Engine/Plugins/Online/Experimental/OnlineServices/Source/OnlineServicesInterface/Public/Online/Friends.h
new file mode 100644
--- /dev/null
+++ b/Engine/Plugins/Online/Experimental/OnlineServices/Source/OnlineServicesInterface/Public/Online/Friends.h
@@ -0,0 +1,101 @@
+// Copyright Epic Games, Inc. All Rights Reserved.
+
+#pragma once
+
+#include "Online/OnlineAsyncOpHandle.h"
+#include "Online/OnlineId.h"
+#include "Online/OnlineMeta.h"
+
+namespace UE::Online {
+
+struct FQueryFriends
+{
+	/** Input struct for Friends::QueryFriends */
+	struct Params
+	{
+		/** Account Id of the local user making the request */
+		FAccountId LocalUserId;
+	};
+
+	/** Output struct for Friends::QueryFriends */
+	struct Result
+	{
+	};
+};
+
+struct FGetFriends
+{
+	/** Input struct for Friends::GetFriends */
+	struct Params
+	{
+		/** Account Id of the local user making the request */
+		FAccountId LocalUserId;
+	};
+
+	/** Output struct for Friends::GetFriends */
+	struct Result
+	{
+		/** Array of friend account ids */
+		TArray<FAccountId> FriendIds;
+	};
+};
+
+/** Struct for FriendsListUpdate event */
+struct FFriendsListUpdated
+{
+	/** Local user who's friends list changed */
+	FOnlineId LocalUserId;
+};
+
+class IFriends
+{
+public:
+	/**
+	 * Query the friends list
+	 * 
+	 * @param Params for the QueryFriends call
+	 * @return
+	 */
+	virtual TOnlineAsyncOpHandle<FQueryFriends> QueryFriends(FQueryFriends::Params&& Params) = 0;
+
+	/**
+	 * Get the contents of a previously queried friends list
+	 *
+	 * @param Params for the GetFriends call
+	 * @return
+	 */
+	virtual TOnlineResult<FGetFriends::Result> GetFriends(FGetFriends::Params&& Params) = 0;
+
+	/**
+	 * Get the event that is triggered when a friends list is updated
+	 *
+	 * @return Event that can be bound to
+	 */
+	virtual TOnlineEvent<void(const FFriendsListUpdated&)> OnFriendsListUpdated() = 0;
+};
+
+namespace Meta {
+// TODO: Move to Friends_Meta.inl file?
+
+BEGIN_ONLINE_STRUCT_META(FQueryFriends::Params)
+	ONLINE_STRUCT_FIELD(FQueryFriends::Params, LocalUserId)
+END_ONLINE_STRUCT_META()
+
+BEGIN_ONLINE_STRUCT_META(FQueryFriends::Result)
+END_ONLINE_STRUCT_META()
+
+BEGIN_ONLINE_STRUCT_META(FGetFriends::Params)
+	ONLINE_STRUCT_FIELD(FGetFriends::Params, LocalUserId)
+END_ONLINE_STRUCT_META()
+
+BEGIN_ONLINE_STRUCT_META(FGetFriends::Result)
+	ONLINE_STRUCT_FIELD(FGetFriends::Result, FriendIds)
+END_ONLINE_STRUCT_META()
+
+BEGIN_ONLINE_STRUCT_META(FFriendsListUpdated)
+	ONLINE_STRUCT_FIELD(FFriendsListUpdated, LocalUserId)
+END_ONLINE_STRUCT_META()
+
+/* Meta*/ }
+
+/* UE::Online */ }
diff --git a/Engine/Plugins/Online/Experimental/OnlineServices/Source/OnlineServicesInterface/Public/Online/OnlineAsyncOpHandle.h b/Engine/Plugins/Online/Experimental/OnlineServices/Source/OnlineServicesInterface/Public/Online/OnlineAsyncOpHandle.h
new file mode 100644
--- /dev/null
+++ b/Engine/Plugins/Online/Experimental/OnlineServices/Source/OnlineServicesInterface/Public/Online/OnlineAsyncOpHandle.h
@@ -0,0 +1,452 @@
+// Copyright Epic Games, Inc. All Rights Reserved.
+
+#pragma once
+
+#include "CoreMinimal.h"
+#include "Online/OnlineError.h"
+#include "Online/OnlineResult.h"
+#include "Templates/Models.h"
+
+#include "Templates/SharedPointer.h"
+#include "Delegates/Delegate.h"
+
+#include <type_traits>
+
+class UObjectBase;
+
+namespace UE::Online {
+
+namespace Private
+{
+	struct CAsSharedCallable
+	{
+		template <typename T>
+		auto Requires(T&& Object) -> decltype(Object->AsShared());
+	};
+
+	template <typename ObjectType, typename = void>
+	struct TSharedPtrHelper
+	{
+		static inline constexpr bool bIsThreadSafe = false;
+		static inline constexpr bool bIsSharedPtr = false;
+	};
+
+	template <typename T>
+	struct TSharedPtrHelper2
+	{
+		using PtrType = decltype(std::remove_pointer_t<T>().AsShared());
+		static inline constexpr bool bIsThreadSafe = TSharedPtrHelper<PtrType>::bIsThreadSafe;
+		static inline constexpr bool bIsSharedPtr = true;
+		using WeakPtrType = typename TSharedPtrHelper<PtrType>::WeakPtrType;
+		static inline WeakPtrType GetWeak(T Ptr) { return Ptr->AsShared(); }
+	};
+
+	template <typename T, ESPMode SPMode>
+	struct TSharedPtrHelper2<TSharedRef<T, SPMode>>
+	{
+		static inline constexpr bool bIsThreadSafe = SPMode == ESPMode::ThreadSafe;
+		static inline constexpr bool bIsSharedPtr = true;
+		using WeakPtrType = TWeakPtr<T, SPMode>;
+		static inline WeakPtrType GetWeak(const TSharedRef<T, SPMode>& Ptr) { return Ptr; }
+	};
+
+	template <typename T>
+	struct TSharedPtrHelper<T, std::enable_if_t<TModels<CAsSharedCallable, T>::Value, void>> : public TSharedPtrHelper2<T>
+	{
+	};
+
+	template <typename DelegateSignature, typename CallableType, typename... VarTypes>
+	TDelegate<DelegateSignature> ConstructFunctionDelegate(CallableType&& Callable, VarTypes&&... Vars)
+	{
+		if constexpr (std::is_function_v<std::remove_reference_t<std::remove_pointer_t<CallableType>>>)
+		{
+			return TDelegate<DelegateSignature>::CreateStatic(Forward<CallableType>(Callable), Forward<VarTypes>(Vars)...);
+		}
+		else
+		{
+			return TDelegate<DelegateSignature>::CreateLambda(Forward<CallableType>(Callable), Forward<VarTypes>(Vars)...);
+		}
+	}
+
+	template <typename DelegateSignature, typename ObjectType, typename CallableType, typename... VarTypes>
+	TDelegate<DelegateSignature> ConstructObjectDelegate(ObjectType Object, CallableType&& Callable, VarTypes&&... Vars)
+	{
+		if constexpr (std::is_member_function_pointer_v<CallableType>)
+		{
+			if constexpr (std::is_convertible_v<ObjectType, const UObjectBase*>)
+			{
+				return TDelegate<DelegateSignature>::CreateUObject(Object, Forward<decltype(Callable)>(Callable), Forward<VarTypes>(Vars)...);
+			}
+			else if constexpr (TSharedPtrHelper<ObjectType>::bIsSharedPtr)
+			{
+				if constexpr (TSharedPtrHelper<ObjectType>::bIsThreadSafe)
+				{
+					return TDelegate<DelegateSignature>::CreateThreadSafeSP(Object, Callable, Forward<VarTypes>(Vars)...);
+				}
+				else
+				{
+					return TDelegate<DelegateSignature>::CreateSP(Object, Callable, Forward<VarTypes>(Vars)...);
+				}
+			}
+			else
+			{
+				return TDelegate<DelegateSignature>::CreateRaw(Object, Callable, Forward<VarTypes>(Vars)...);
+			}
+		}
+		else
+		{
+			if constexpr (std::is_convertible_v<ObjectType, const UObjectBase*>)
+			{
+				if constexpr (std::is_same_v<FName, std::remove_cv_t<CallableType>>)
+				{
+					return TDelegate<DelegateSignature>::CreateUFunction(Object, Callable, Forward<VarTypes>(Vars)...);
+				}
+				else
+				{
+					return TDelegate<DelegateSignature>::CreateWeakLambda(Object, Callable, Forward<VarTypes>(Vars)...);
+				}
+			}
+			else
+			{
+				static_assert(TSharedPtrHelper<ObjectType>::bIsSharedPtr, "When using non pointer to member functions, the first parameter can only be a UObject*, TSharedRef, or pointer to a class that derives from TSharedFromThis");
+				auto InnerDelegate = ConstructFunctionDelegate<DelegateSignature>(Callable, Forward<VarTypes>(Vars)...);
+				return TDelegate<DelegateSignature>::CreateLambda([WeakObject = TSharedPtrHelper<ObjectType>::GetWeak(Object), Delegate = MoveTemp(InnerDelegate)](auto&&... Params)
+					{
+						auto StrongObject = WeakObject.Pin();
+						if (StrongObject.IsValid())
+						{
+							Delegate.ExecuteIfBound(Forward<decltype(Params)>(Params)...);
+						}
+					});
+			}
+		}
+	}
+
+	template <typename DelegateSignature, typename ObjectOrCallableType, typename... VarTypes>
+	UE_NODISCARD inline TDelegate<DelegateSignature> ConstructDelegate(ObjectOrCallableType&& ObjectOrCallable, VarTypes&&... Vars)
+	{
+		if constexpr((std::is_pointer_v<std::remove_reference_t<ObjectOrCallableType>> && !std::is_function_v<std::remove_reference_t<std::remove_pointer_t<ObjectOrCallableType>>>) || TSharedPtrHelper<std::remove_reference_t<ObjectOrCallableType>>::bIsSharedPtr)
+		{
+			return ConstructObjectDelegate<DelegateSignature>(Forward<ObjectOrCallableType>(ObjectOrCallable), Forward<VarTypes>(Vars)...);
+		}
+		else
+		{
+			return ConstructFunctionDelegate<DelegateSignature>(Forward<ObjectOrCallableType>(ObjectOrCallable), Forward<VarTypes>(Vars)...);
+		}
+	}
+
+/* Private */ }
+
+template<typename DelegateSignature>
+class TOnlineCallback
+{
+	// Will hit TDelegate's static_asserts
+	TDelegate<DelegateSignature> Delegate;
+};
+
+/**
+ *
+ */
+template <typename InRetValType, typename... ParamTypes>
+class TOnlineCallback<InRetValType(ParamTypes...)>
+{
+public:
+	using DelegateSignature = InRetValType(ParamTypes...);
+
+	/**
+	 * 
+	 */
+	template<typename... ConstructDelegateParamTypes>
+	void Bind(ConstructDelegateParamTypes&&... Params)
+	{
+		Delegate = Private::ConstructDelegate<DelegateSignature>(Forward<ConstructDelegateParamTypes...>(Params)...);
+	}
+
+	/**
+	 *
+	 */
+	void Unbind()
+	{
+		Delegate.Unbind();
+	}
+
+	/**
+	 *
+	 */
+	void ExecuteIfBound(ParamTypes&&... Params)
+	{
+		Delegate.ExecuteIfBound(Forward<ParamTypes...>(Params)...);
+	}
+
+private:
+	TDelegate<DelegateSignature> Delegate;
+};
+
+/**
+ *
+ */
+class FOnlineEventDelegateHandle
+{
+public:
+	/**
+	 *
+	 */
+	template <typename DelegateSignature>
+	FOnlineEventDelegateHandle(FDelegateHandle Handle, TSharedRef<TMulticastDelegate<DelegateSignature>>& DelegateRef)
+	{
+		DelegateUnbinder = [Handle, DelegateWeakPtr = TWeakPtr<TMulticastDelegate<DelegateSignature>>(DelegateRef)]()
+			{
+				TSharedPtr<TMulticastDelegate<DelegateSignature>> DelegateSP = DelegateWeakPtr.Pin();
+				if (DelegateSP)
+				{
+					check(IsInGameThread());
+					DelegateSP->Remove(Handle);
+				}
+			};
+	}
+
+	/**
+	 *
+	 */
+	~FOnlineEventDelegateHandle()
+	{
+		Unbind();
+	}
+
+	/**
+	 * 
+	 */
+	void Unbind()
+	{
+		if (DelegateUnbinder)
+		{
+			DelegateUnbinder();
+			DelegateUnbinder = nullptr;
+		}
+	}
+
+private:
+	FOnlineEventDelegateHandle(const FOnlineEventDelegateHandle&) = delete;
+	FOnlineEventDelegateHandle& operator=(const FOnlineEventDelegateHandle&) = delete;
+
+	TFunction<void()> DelegateUnbinder;
+};
+
+template<typename DelegateSignature>
+class TOnlineEventCallable
+{
+	// Will hit TMulticastDelegate's static_asserts
+	TMulticastDelegate<DelegateSignature> Delegate;
+};
+
+/**
+ *
+ */
+template <typename... ParamTypes>
+class TOnlineEventCallable<void(ParamTypes...)>
+{
+public:
+	using DelegateSignature = void(ParamTypes...);
+
+	/**
+	 *
+	 */
+	TOnlineEventCallable()
+		: DelegateRef(MakeShared<TMulticastDelegate<DelegateSignature>>())
+	{
+	}
+
+	/**
+	 *
+	 */
+	template<typename... ConstructDelegateParamTypes>
+	FOnlineEventDelegateHandle Add(ConstructDelegateParamTypes&&... Params)
+	{
+		FDelegateHandle Handle = DelegateRef->Add(Private::ConstructDelegate<DelegateSignature>(Forward<ConstructDelegateParamTypes...>(Params)...));
+		return FOnlineEventDelegateHandle(DelegateRef, Handle);
+	}
+
+	/**
+	 *
+	 */
+	void Broadcast(ParamTypes&&... Params)
+	{
+		DelegateRef->Broadcast(Forward<ParamTypes...>(Params)...);
+	}
+
+private:
+	TSharedRef<TMulticastDelegate<DelegateSignature>> DelegateRef;
+};
+
+/**
+ *
+ */
+template <typename DelegateSignature>
+class TOnlineEvent
+{
+public:
+	/**
+	 *
+	 */
+	TOnlineEvent(TOnlineEventCallable<DelegateSignature>& InEvent)
+		: Event(InEvent)
+	{
+	}
+
+	/**
+	 *
+	 */
+	template<typename... ConstructDelegateParamTypes>
+	FOnlineEventDelegateHandle Add(ConstructDelegateParamTypes&&... Params)
+	{
+		return Event.Add(Forward<ConstructDelegateParamTypes...>(Params)...);
+	}
+
+private:
+	TOnlineEventCallable<DelegateSignature>& Event;
+};
+
+enum class EAsyncOpState : uint8
+{
+	Invalid,
+	Queued,
+	Running,
+	RetryQueued,
+	Complete,
+	Cancelled
+};
+
+class IOnlineAsyncOp
+{
+public:
+	virtual void Cancel() = 0;
+	virtual EAsyncOpState GetState() const = 0;
+	virtual void DetachHandle(class IAsyncHandle* Handle) = 0;
+};
+
+struct FAsyncProgress
+{
+	/* Step that is running */
+	const FString Step;
+	/* How much has been completed */
+	double Progress;
+	/* How much total there is to do */
+	double MaxProgress;
+	/* Localized message suitable for showing to end user */
+	FText ProgressMessage;
+};
+
+struct FWillRetry
+{
+	FOnlineError Reason;
+	double RetryInSeconds;
+	int RetryAttempt;
+	int MaxRetryAttempts;
+};
+
+template <typename OpType> class TOnlineAsyncOpHandle;
+
+namespace Private
+{
+	// State shared between the Op and Handle
+	template <typename OpType>
+	class IOnlineAsyncOpSharedState
+	{
+	public:
+		virtual ~IOnlineAsyncOpSharedState() {}
+		virtual void Cancel(const FOnlineError& Reason) = 0;
+		virtual EAsyncOpState GetState() const = 0;
+		virtual void SetOnProgress(TDelegate<void(const FAsyncProgress&)>&& Delegate) = 0;
+		virtual void SetOnWillRetry(TDelegate<void(TOnlineAsyncOpHandle<OpType>&, const FWillRetry&)>&& Delegate) = 0;
+		virtual void SetOnComplete(TDelegate<void(const TOnlineResult<typename OpType::Result>&)>&& Delegate) = 0;
+	};
+}
+
+/**
+ *
+ */
+template <typename OpType>
+class TOnlineAsyncOpHandle
+{
+public:
+	/**
+	 *
+	 */
+	TOnlineAsyncOpHandle(TSharedRef<Private::IOnlineAsyncOpSharedState<OpType>>&& InSharedState)
+		: State(MoveTemp(InSharedState))
+	{
+	}
+
+	// need to add a comment here explaining how this works:
+	// There's an optional first parameter that is a class instance (raw pointer or shared/weakptr/ref/weakobjptr). If possible, this will be checked for validity before calling the callback
+	// Next is the callback, signatures in FAsyncOpState (TODO: Move that to a more visible location), can be a member function
+	// After that is any additional params that should be captured (like the ue4 delegates support) (or we remove this capability)
+	// Always called on the game thread
+
+	/**
+	 * 
+	 *
+	 * @param 
+	 * @return The async op handle
+	 */
+	template <typename... ArgTypes>
+	TOnlineAsyncOpHandle& OnProgress(ArgTypes&&... Args)
+	{
+		State->SetOnProgress(Private::ConstructDelegate<void(const FAsyncProgress&)>(Forward<ArgTypes>(Args)...));
+
+		return *this;
+	}
+
+	/**
+	 * 
+	 *
+	 * @param 
+	 * @return The async op handle
+	 */
+	template <typename... ArgTypes>
+	TOnlineAsyncOpHandle& OnWillRetry(ArgTypes&&... Args)
+	{
+		State->SetOnWillRetry(Private::ConstructDelegate<void(TOnlineAsyncOpHandle<OpType>&, const FWillRetry&)>(Forward<ArgTypes>(Args)...));
+
+		return *this;
+	}
+
+	/**
+	 * 
+	 *
+	 * @param
+	 * @return The async op handle
+	 */
+	template <typename... ArgTypes>
+	TOnlineAsyncOpHandle& OnComplete(ArgTypes&&... Args)
+	{
+		State->SetOnComplete(Private::ConstructDelegate<void(const TOnlineResult<typename OpType::Result>&)>(Forward<ArgTypes>(Args)...));
+
+		return *this;
+	}
+
+	/**
+	 * Cancel the operation
+	 *
+	 * @param Reason Allows a more specific reson why the operation is cancelled to be specified
+	 */
+	void Cancel(const FOnlineError& Reason = Errors::Cancelled())
+	{
+		State->Cancel(Reason);
+	}
+
+	/**
+	 * Get the state of the operation
+	 * 
+	 * @return The state of the operation
+	 */
+	EAsyncOpState GetState() const
+	{
+		return State->GetState();
+	}
+
+private:
+	TSharedPtr<Private::IOnlineAsyncOpSharedState<OpType>> State;
+};
+
+/* UE::Online */ }
diff --git a/Engine/Plugins/Online/Experimental/OnlineServices/Source/OnlineServicesInterface/Public/Online/OnlineError.h b/Engine/Plugins/Online/Experimental/OnlineServices/Source/OnlineServicesInterface/Public/Online/OnlineError.h
new file mode 100644
--- /dev/null
+++ b/Engine/Plugins/Online/Experimental/OnlineServices/Source/OnlineServicesInterface/Public/Online/OnlineError.h
@@ -0,0 +1,27 @@
+// Copyright Epic Games, Inc. All Rights Reserved.
+
+#pragma once
+
+/**
+ * This file is a placeholder, FOnlineError will be implemented in the future
+ */
+
+namespace UE::Online {
+
+class FOnlineError
+{
+public:
+	FOnlineError() = default;
+	FOnlineError(FString&& InError) : Error(MoveTemp(InError)) {}
+	bool operator==(const FOnlineError& Other) const { return Error == Other.Error; }
+
+	FString Error;
+};
+
+namespace Errors
+{
+	inline FOnlineError Cancelled() { return FOnlineError(); }
+	inline FOnlineError Unimplemented() { return FOnlineError(); }
+}
+
+/* UE::Online */ }
diff --git a/Engine/Plugins/Online/Experimental/OnlineServices/Source/OnlineServicesInterface/Public/Online/OnlineId.h b/Engine/Plugins/Online/Experimental/OnlineServices/Source/OnlineServicesInterface/Public/Online/OnlineId.h
new file mode 100644
--- /dev/null
+++ b/Engine/Plugins/Online/Experimental/OnlineServices/Source/OnlineServicesInterface/Public/Online/OnlineId.h
@@ -0,0 +1,39 @@
+// Copyright Epic Games, Inc. All Rights Reserved.
+
+#pragma once
+
+#include "Templates/TypeHash.h"
+#include "UObject/NameTypes.h"
+
+/**
+ * This file is a placeholder, will be implemented along with the id registry in the future
+ */
+
+namespace UE::Online {
+
+using FOnlineIdHandle = int32;
+struct FOnlineId
+{
+	FName Type;
+	FOnlineIdHandle Handle;
+};
+
+using FAccountId = FOnlineId;
+
+inline bool operator==(const FOnlineId& A, const FOnlineId& B)
+{
+	return A.Type == B.Type && A.Handle == B.Handle;
+}
+
+inline bool operator!=(const FOnlineId& A, const FOnlineId& B)
+{
+	return !(A == B);
+}
+
+inline uint32 GetTypeHash(const FOnlineId& AccountId)
+{
+	using ::GetTypeHash;
+	return HashCombine(GetTypeHash(AccountId.Type), GetTypeHash(AccountId.Handle));
+}
+
+/* UE::Online */ }
diff --git a/Engine/Plugins/Online/Experimental/OnlineServices/Source/OnlineServicesInterface/Public/Online/OnlineResult.h b/Engine/Plugins/Online/Experimental/OnlineServices/Source/OnlineServicesInterface/Public/Online/OnlineResult.h
new file mode 100644
--- /dev/null
+++ b/Engine/Plugins/Online/Experimental/OnlineServices/Source/OnlineServicesInterface/Public/Online/OnlineResult.h
@@ -0,0 +1,202 @@
+// Copyright Epic Games, Inc. All Rights Reserved.
+
+#pragma once
+
+#include "Misc/TVariant.h"
+
+#include "Online/OnlineError.h"
+
+namespace UE::Online {
+
+/**
+ * A container that holds either a successful result or an erroneous result.
+ *
+ * @param SuccessType - the type to use for successful results
+ * @param ErrorType - the type to use for erroneous results
+ */
+template <typename SuccessType, typename ErrorType>
+class TResult
+{
+public:
+	/**
+	 * Construct a result with an Ok value
+	 */
+	explicit TResult(const SuccessType& OkValue)
+		: Storage(TInPlaceType<SuccessType>(), OkValue)
+	{
+	}
+
+	/**
+	 * Construct a result with an Ok value
+	 */
+	explicit TResult(SuccessType&& OkValue)
+		: Storage(TInPlaceType<SuccessType>(), MoveTemp(OkValue))
+	{
+	}
+
+	/**
+	 * Construct a result with an error value
+	 */
+	explicit TResult(const ErrorType& ErrValue)
+		: Storage(TInPlaceType<ErrorType>(), ErrValue)
+	{
+	}
+
+	/**
+	 * Construct a result with an error value
+	 */
+	explicit TResult(ErrorType&& ErrValue)
+		: Storage(TInPlaceType<ErrorType>(), MoveTemp(ErrValue))
+	{
+	}
+
+	/**
+	 * Copy construct a result of the same type
+	 */
+	TResult(const TResult& Other) = default;
+
+	/**
+	 * Move construct a result
+	 */
+	TResult(TResult&& Other) = default;
+
+	/**
+	 * Assign a result from another result
+	 */
+	TResult& operator=(const TResult& Other)
+	{
+		if (&Other != this)
+		{
+			Storage = Other.Storage;
+		}
+		return *this;
+	}
+
+	/**
+	 * Move-assign a result
+	 */
+	TResult& operator=(TResult&& Other)
+	{
+		if (&Other != this)
+		{
+			Storage = MoveTemp(Other.Storage);
+		}
+		return *this;
+	}
+
+	virtual ~TResult() = default;
+
+	/**
+	 * Check if the value held in the result is a SuccessType.
+	 *
+	 * If the held value is a SuccessType, it is permissible to call GetOkValue. Otherwise it is permissible to call GetErrorValue.
+	 *
+	 * @return true if the held value is a SuccessType, false otherwise.
+	 */
+	bool IsOk() const
+	{
+		return Storage.template IsType<SuccessType>();
+	}
+
+	/**
+	 * Check if the value held in the result is a ErrorType.
+	 *
+	 * If the held value is a ErrorType, it is permissible to call GetErrorValue. Otherwise it is permissible to call GetOkValue.
+	 *
+	 * @return true if the held value is a ErrorType, false otherwise.
+	 */
+	bool IsError() const
+	{
+		return Storage.template IsType<ErrorType>();
+	}
+
+	/** Get the Ok value stored in the result. This must not be called on a result holding the error type */
+	const SuccessType& GetOkValue() const
+	{
+		checkf(IsOk(), TEXT("It is an error to call GetOkValue() on a TResult that does not hold an ok value. Please either check IsOk() or use TryGetOkValue"));
+		return Storage.template Get<SuccessType>();
+	}
+
+	/** Get the Ok value stored in the result. This must not be called on a result holding the error type */
+	SuccessType& GetOkValue()
+	{
+		checkf(IsOk(), TEXT("It is an error to call GetOkValue() on a TResult that does not hold an ok value. Please either check IsOk() or use TryGetOkValue"));
+		return Storage.template Get<SuccessType>();
+	}
+
+	/** Get the Error value stored in the result. This must not be called on a result holding the success type */
+	const ErrorType& GetErrorValue() const
+	{
+		checkf(IsError(), TEXT("It is an error to call GetErrorValue() on a TResult that does not hold an error value. Please either check IsError() or use TryGetErrorValue"));
+		return Storage.template Get<ErrorType>();
+	}
+
+	/** Get the Error value stored in the result. This must not be called on a result holding the success type */
+	ErrorType& GetErrorValue()
+	{
+		checkf(IsError(), TEXT("It is an error to call GetErrorValue() on a TResult that does not hold an error value. Please either check IsError() or use TryGetErrorValue"));
+		return Storage.template Get<ErrorType>();
+	}
+
+	/**
+	 * Convert from TResult<Success, Error> to Success* if the result is successful.
+	 *
+	 * @return A pointer to a value of type SuccessType if the result is successful, otherwise nullptr
+	 */
+	const SuccessType* TryGetOkValue() const
+	{
+		return const_cast<TResult*>(this)->TryGetOkValue();
+	}
+
+	/**
+	 * Convert from TResult<Success, Error> to Success* if the result is successful.
+	 *
+	 * @return A pointer to a value of type SuccessType if the result is successful, otherwise nullptr
+	 */
+	SuccessType* TryGetOkValue()
+	{
+		return Storage.template TryGet<SuccessType>();
+	}
+
+	/**
+	 * Convert from TResult<Success, Error> to Error* if the result is erroneous.
+	 *
+	 * @return A pointer to a value of type ErrorType if the result is erroneous, otherwise nullptr
+	 */
+	const ErrorType* TryGetErrorValue() const
+	{
+		return const_cast<TResult*>(this)->TryGetErrorValue();
+	}
+
+	/**
+	 * Convert from TResult<Success, Error> to Error* if the result is erroneous.
+	 *
+	 * @return A pointer to a value of type ErrorType if the result is erroneous, otherwise nullptr
+	 */
+	ErrorType* TryGetErrorValue()
+	{
+		return Storage.template TryGet<ErrorType>();
+	}
+
+	/**
+	 * Unwraps the result, returning the success value if one is held, otherwise returning the default passed.
+	 *
+	 * @param DefaultValue - The value to return if the result does not hold a success type.
+	 */
+	const SuccessType& GetOkOrDefault(const SuccessType& DefaultValue) const
+	{
+		return IsOk() ? GetOkValue() : DefaultValue;
+	}
+
+
+private:
+	TResult() = delete;
+
+	/** Location that the result's value is stored */
+	TVariant<SuccessType, ErrorType> Storage;
+};
+
+template <typename SuccessType>
+using TOnlineResult = TResult<SuccessType, FOnlineError>;
+
+/* UE::Online */ }
diff --git a/Engine/Plugins/Online/Experimental/OnlineServices/Source/OnlineServicesInterface/Public/Online/OnlineServices.h b/Engine/Plugins/Online/Experimental/OnlineServices/Source/OnlineServicesInterface/Public/Online/OnlineServices.h
new file mode 100644
--- /dev/null
+++ b/Engine/Plugins/Online/Experimental/OnlineServices/Source/OnlineServicesInterface/Public/Online/OnlineServices.h
@@ -0,0 +1,123 @@
+// Copyright Epic Games, Inc. All Rights Reserved.
+
+#pragma once
+
+#include "Templates/SharedPointer.h"
+
+namespace UE::Online {
+
+enum class EOnlineServices : uint8
+{
+	// Null, Providing minimal functionality when no backend services are required
+	Null,
+	// Epic Online Services
+	Epic,
+	// Xbox services
+	Xbox,
+	// PlayStation Network
+	PSN,
+	// Nintendo
+	Nintendo,
+	// Stadia,
+	Stadia,
+	// Steam
+	Steam,
+	// Google
+	Google,
+	// GooglePlay
+	GooglePlay,
+	// Apple
+	Apple,
+	// GameKit
+	AppleGameKit,
+	// Samsung
+	Samsung,
+	// Oculus
+	Oculus,
+	// Tencent
+	Tencent,
+	// Reserved for future use/platform extensions
+	Reserved_14,
+	Reserved_15,
+	Reserved_16,
+	Reserved_17,
+	Reserved_18,
+	Reserved_19,
+	Reserved_20,
+	Reserved_21,
+	Reserved_22,
+	Reserved_23,
+	Reserved_24,
+	Reserved_25,
+	Reserved_26,
+	Reserved_27,
+	// For game specific Online Services
+	GameDefined_0 = 28,
+	GameDefined_1,
+	GameDefined_2,
+	GameDefined_3,
+	// Platform native, may not exist for all platforms
+	Platform = 254,
+	// Default, configured via ini, TODO: List specific ini section/key
+	Default = 255
+};
+
+// Interfaces
+using IAuthPtr = TSharedPtr<class IAuth>;
+using IFriendsPtr = TSharedPtr<class IFriends>;
+
+class ONLINESERVICESINTERFACE_API IOnlineServices
+{
+public:
+	/**
+	 *
+	 */
+	virtual void Init() = 0;
+
+	/**
+	 *
+	 */
+	virtual void Destroy() = 0;
+
+	/**
+	 *
+	 */
+	virtual IAuthPtr GetAuthInterface() = 0;
+
+	/**
+	 *
+	 */
+	virtual IFriendsPtr GetFriendsInterface() = 0;
+};
+
+/**
+ * Get an instance of the online subsystem
+ *
+ * @param OnlineServices Type of online services to retrieve
+ * @param InstanceName Name of the services instance to retrieve
+ * @return The services instance or an invalid pointer if the services is unavailable
+ */
+ONLINESERVICESINTERFACE_API TSharedPtr<IOnlineServices> GetServices(EOnlineServices OnlineServices = EOnlineServices::Default, FName InstanceName = NAME_None);
+
+/**
+ * Get a specific services type and cast to the specific services type
+ *
+ * @param InstanceName Name of the services instance to retrieve
+ * @return The services instance or an invalid pointer if the services is unavailable
+ */
+template <typename ServicesClass>
+TSharedPtr<ServicesClass> GetServices(FName InstanceName = NAME_None)
+{
+	return StaticCastSharedPtr<ServicesClass>(GetServices(ServicesClass::GetServicesProvider(), InstanceName));
+}
+
+/**
+ * Destroy an instance of the online subsystem
+ *
+ * @param OnlineServices Type of online services to destroy
+ * @param InstanceName Name of the services instance to destroy
+ */
+ONLINESERVICESINTERFACE_API void DestroyServices(EOnlineServices OnlineServices = EOnlineServices::Default, FName InstanceName = NAME_None);
+
+
+/* UE::Online */ }
diff --git a/Engine/Plugins/Online/Experimental/OnlineServices/Source/OnlineServicesInterface/Public/Online/OnlineServicesRegistry.h b/Engine/Plugins/Online/Experimental/OnlineServices/Source/OnlineServicesInterface/Public/Online/OnlineServicesRegistry.h
new file mode 100644
--- /dev/null
+++ b/Engine/Plugins/Online/Experimental/OnlineServices/Source/OnlineServicesInterface/Public/Online/OnlineServicesRegistry.h
@@ -0,0 +1,112 @@
+// Copyright Epic Games, Inc. All Rights Reserved.
+
+#pragma once
+
+#include "Containers/Map.h"
+#include "Templates/SharedPointer.h"
+#include "Templates/UniquePtr.h"
+
+#include "Online/OnlineServices.h"
+
+class FLazySingleton;
+
+namespace UE::Online {
+	
+enum class EOnlineServices : uint8;
+
+class ONLINESERVICESINTERFACE_API IOnlineServicesFactory
+{
+public:
+	/**
+	 * Virtual destructor
+	 */
+	virtual ~IOnlineServicesFactory() {}
+
+	/**
+	 * Create an IOnlineServices instance
+	 *
+	 * @return Initialized IOnlineServices instance
+	 */
+	virtual TSharedPtr<IOnlineServices> Create() = 0;
+};
+
+class FOnlineServicesRegistry
+{
+public:
+	/**
+	 * Get the FOnlineServicesRegistry singleton
+	 * 
+	 * @return The FOnlineServicesRegistry singleton instance
+	 */
+	ONLINESERVICESINTERFACE_API static FOnlineServicesRegistry& Get();
+
+	/**
+	 * Tear down the singleton instance
+	 */
+	ONLINESERVICESINTERFACE_API static void TearDown();
+
+	/**
+	 * Register a factory for creation of IOnlineServices instances
+	 * 
+	 * @param OnlineServices Services that the factory is for
+	 * @param Factory Factory for creation of IOnlineServices instances
+	 * @param Priority Integer priority, allows an existing OnlineServices implementation to be extended and registered with a higher priority so it is used instead
+	 */
+	ONLINESERVICESINTERFACE_API void RegisterServicesFactory(EOnlineServices OnlineServices, TUniquePtr<IOnlineServicesFactory>&& Factory, int32 Priority = 0);
+
+	/**
+	 * Unregister a previously registered IOnlineServices factory
+	 *
+	 * @param OnlineServices Services that the factory is for
+	 * @param Priority Integer priority, will be unregistered only if the priority matches the one that is registered
+	 */
+	ONLINESERVICESINTERFACE_API void UnRegisterServicesFactory(EOnlineServices OnlineServices, int32 Priority = 0);
+
+	/**
+	 * Get a named instance of a specific IOnlineServices
+	 * 
+	 * @param OnlineServices Type of online services for the IOnlineServices instance
+	 * @param InstanceName Name of the instance
+	 * 
+	 * @return The services instance, or an invalid pointer if the OnlineServices is unavailable
+	 */
+	ONLINESERVICESINTERFACE_API TSharedPtr<IOnlineServices> GetNamedServicesInstance(EOnlineServices OnlineServices, FName InstanceName);
+
+	/**
+	 * Destroy a named instance of a specific OnlineServices
+	 *
+	 * @param OnlineServices  Type of online services for the IOnlineServices instance
+	 * @param InstanceName Name of the instance
+	 */
+	ONLINESERVICESINTERFACE_API void DestroyNamedServicesInstance(EOnlineServices OnlineServices, FName InstanceName);
+
+	/**
+	 * Create and initialize a new IOnlineServices instance
+	 *
+	 * @param OnlineServices Type of online services for the IOnlineServices instance
+	 * 
+	 * @return The initialized IOnlineServices instance, or an invalid pointer if the OnlineServices is unavailable
+	 */
+	ONLINESERVICESINTERFACE_API TSharedPtr<IOnlineServices> CreateServices(EOnlineServices OnlineServices);
+
+private:
+	struct FFactoryAndPriority
+	{
+		FFactoryAndPriority(TUniquePtr<IOnlineServicesFactory>&& InFactory, int32 InPriority)
+			: Factory(MoveTemp(InFactory))
+			, Priority(InPriority)
+		{
+		}
+
+		TUniquePtr<IOnlineServicesFactory> Factory;
+		int32 Priority;
+	};
+
+	TMap<EOnlineServices, FFactoryAndPriority> ServicesFactories;
+	TMap<EOnlineServices, TMap<FName, TSharedPtr<IOnlineServices>>> NamedServiceInstances;
+
+	FOnlineServicesRegistry() {}
+	friend FLazySingleton;
+};
+
+/* UE::Online */ }
