- Fix broken FOnlineAccountIdHandle constructor

diff --git a/Engine/Plugins/Online/Experimental/OnlineServices/Source/OnlineServicesInterface/Public/Online/OnlineId.h b/Engine/Plugins/Online/Experimental/OnlineServices/Source/OnlineServicesInterface/Public/Online/OnlineId.h
--- a/Engine/Plugins/Online/Experimental/OnlineServices/Source/OnlineServicesInterface/Public/Online/OnlineId.h
+++ b/Engine/Plugins/Online/Experimental/OnlineServices/Source/OnlineServicesInterface/Public/Online/OnlineId.h
@@ -27,7 +27,11 @@ class TOnlineIdHandle
 {
 public:
 	TOnlineIdHandle() = default;
-	TOnlineIdHandle(EOnlineServices Type, uint32 Handle) : Value((Handle & 0x00FFFFFF) & (uint32(Type) << 24)) {}
+	TOnlineIdHandle(EOnlineServices Type, uint32 Handle)
+	{
+		check(Handle < 0xFF000000);
+		Value = (Handle & 0x00FFFFFF) | (uint32(Type) << 24);
+	}
 
 	inline bool IsValid() const { return GetHandle() != 0; }
 
