- Rename ToDebugString to ToLogString for consistency

diff --git a/Engine/Plugins/Online/Experimental/OnlineServices/Source/OnlineServicesInterface/Public/Online/OnlineServices.h b/Engine/Plugins/Online/Experimental/OnlineServices/Source/OnlineServicesInterface/Public/Online/OnlineServices.h
--- a/Engine/Plugins/Online/Experimental/OnlineServices/Source/OnlineServicesInterface/Public/Online/OnlineServices.h
+++ b/Engine/Plugins/Online/Experimental/OnlineServices/Source/OnlineServicesInterface/Public/Online/OnlineServices.h
@@ -51,7 +51,7 @@ public:
 	/**
 	 * 
 	 */
-	virtual FString ToDebugString(const FOnlineAccountIdHandle& Handle) = 0;
+	virtual FString ToLogString(const FOnlineAccountIdHandle& Handle) = 0;
 };
 
 /**
@@ -89,7 +89,7 @@ inline FString ToLogString(const TOnlineIdHandle<IdType>& Id)
 	FString Result;
 	if (TSharedPtr<IOnlineServices> Services = GetServices(Id.GetType()))
 	{
-		Result = Services->ToDebugString(Id);
+		Result = Services->ToLogString(Id);
 	}
 	return Result;
 }
diff --git a/Engine/Plugins/Online/Experimental/OnlineServicesEOS/Source/Private/OnlineServicesEOS.cpp b/Engine/Plugins/Online/Experimental/OnlineServicesEOS/Source/Private/OnlineServicesEOS.cpp
--- a/Engine/Plugins/Online/Experimental/OnlineServicesEOS/Source/Private/OnlineServicesEOS.cpp
+++ b/Engine/Plugins/Online/Experimental/OnlineServicesEOS/Source/Private/OnlineServicesEOS.cpp
@@ -88,7 +88,7 @@ void FOnlineServicesEOS::Initialize()
 	FOnlineServicesCommon::Initialize();
 }
 
-FString FOnlineServicesEOS::ToDebugString(const FOnlineAccountIdHandle& Handle)
+FString FOnlineServicesEOS::ToLogString(const FOnlineAccountIdHandle& Handle)
 {
 	FString Result;
 	if (Handle.IsValid())
diff --git a/Engine/Plugins/Online/Experimental/OnlineServicesEOS/Source/Public/OnlineServicesEOS.h b/Engine/Plugins/Online/Experimental/OnlineServicesEOS/Source/Public/OnlineServicesEOS.h
--- a/Engine/Plugins/Online/Experimental/OnlineServicesEOS/Source/Public/OnlineServicesEOS.h
+++ b/Engine/Plugins/Online/Experimental/OnlineServicesEOS/Source/Public/OnlineServicesEOS.h
@@ -26,7 +26,7 @@ public:
 
 	EOS_HPlatform GetEOSPlatformHandle() const;
 
-	virtual FString ToDebugString(const FOnlineAccountIdHandle& Handle) override;
+	virtual FString ToLogString(const FOnlineAccountIdHandle& Handle) override;
 protected:
 	IEOSPlatformHandlePtr EOSPlatformHandle;
 };
