- Rename TOnlineIdHandle::GetType to GetOnlineServicesType

diff --git a/Engine/Plugins/Online/Experimental/OnlineServices/Source/OnlineServicesInterface/Public/Online/OnlineId.h b/Engine/Plugins/Online/Experimental/OnlineServices/Source/OnlineServicesInterface/Public/Online/OnlineId.h
--- a/Engine/Plugins/Online/Experimental/OnlineServices/Source/OnlineServicesInterface/Public/Online/OnlineId.h
+++ b/Engine/Plugins/Online/Experimental/OnlineServices/Source/OnlineServicesInterface/Public/Online/OnlineId.h
@@ -35,7 +35,7 @@ public:
 
 	inline bool IsValid() const { return GetHandle() != 0; }
 
-	EOnlineServices GetType() const { return EOnlineServices(Value >> 24); }
+	EOnlineServices GetOnlineServicesType() const { return EOnlineServices(Value >> 24); }
 	uint32 GetHandle() const { return Value & 0x00FFFFFF; }
 
 	bool operator==(const TOnlineIdHandle& Other) const { return Value == Other.Value; }
@@ -51,7 +51,7 @@ template<EOnlineIdType IdType>
 inline uint32 GetTypeHash(const TOnlineIdHandle<IdType>& Handle)
 {
 	using ::GetTypeHash;
-	return HashCombine(GetTypeHash(Handle.GetType()), GetTypeHash(Handle.GetHandle()));
+	return HashCombine(GetTypeHash(Handle.GetOnlineServicesType()), GetTypeHash(Handle.GetHandle()));
 }
 
 template<EOnlineIdType IdType>
diff --git a/Engine/Plugins/Online/Experimental/OnlineServices/Source/OnlineServicesInterface/Public/Online/OnlineServices.h b/Engine/Plugins/Online/Experimental/OnlineServices/Source/OnlineServicesInterface/Public/Online/OnlineServices.h
--- a/Engine/Plugins/Online/Experimental/OnlineServices/Source/OnlineServicesInterface/Public/Online/OnlineServices.h
+++ b/Engine/Plugins/Online/Experimental/OnlineServices/Source/OnlineServicesInterface/Public/Online/OnlineServices.h
@@ -87,7 +87,7 @@ template<EOnlineIdType IdType>
 inline FString ToLogString(const TOnlineIdHandle<IdType>& Id)
 {
 	FString Result;
-	if (TSharedPtr<IOnlineServices> Services = GetServices(Id.GetType()))
+	if (TSharedPtr<IOnlineServices> Services = GetServices(Id.GetOnlineServicesType()))
 	{
 		Result = Services->ToLogString(Id);
 	}
diff --git a/Engine/Plugins/Online/Experimental/OnlineServicesEOS/Source/Private/OnlineIdEOS.h b/Engine/Plugins/Online/Experimental/OnlineServicesEOS/Source/Private/OnlineIdEOS.h
--- a/Engine/Plugins/Online/Experimental/OnlineServicesEOS/Source/Private/OnlineIdEOS.h
+++ b/Engine/Plugins/Online/Experimental/OnlineServicesEOS/Source/Private/OnlineIdEOS.h
@@ -61,7 +61,7 @@ FOnlineAccountIdHandle CreateAccountId(const EOS_EpicAccountId EpicAccountId, co
 template<EOnlineIdType IdType>
 inline bool ValidateOnlineId(const TOnlineIdHandle<IdType>& Handle)
 {
-	return Handle.GetType() == EOnlineServices::Epic && Handle.IsValid();
+	return Handle.GetOnlineServicesType() == EOnlineServices::Epic && Handle.IsValid();
 }
 
 } /* namespace UE::Online */
