- Change TOnlineIdHandle to use a struct tag instead of an enum for the template argument

diff --git a/Engine/Plugins/Online/Experimental/OnlineServices/Source/OnlineServicesCommon/Public/Online/OnlineIdCommon.h b/Engine/Plugins/Online/Experimental/OnlineServices/Source/OnlineServicesCommon/Public/Online/OnlineIdCommon.h
--- a/Engine/Plugins/Online/Experimental/OnlineServices/Source/OnlineServicesCommon/Public/Online/OnlineIdCommon.h
+++ b/Engine/Plugins/Online/Experimental/OnlineServices/Source/OnlineServicesCommon/Public/Online/OnlineIdCommon.h
@@ -9,7 +9,7 @@ namespace UE::Online {
 /**
  * A net id registry suitable for use with trivial immutable keys
  */
-template<EOnlineIdType IdType, typename IdValueType, EOnlineServices SubsystemType>
+template<typename IdType, typename IdValueType, EOnlineServices SubsystemType>
 class TOnlineBasicIdRegistry
 {
 public:
diff --git a/Engine/Plugins/Online/Experimental/OnlineServices/Source/OnlineServicesInterface/Public/Online/OnlineId.h b/Engine/Plugins/Online/Experimental/OnlineServices/Source/OnlineServicesInterface/Public/Online/OnlineId.h
--- a/Engine/Plugins/Online/Experimental/OnlineServices/Source/OnlineServicesInterface/Public/Online/OnlineId.h
+++ b/Engine/Plugins/Online/Experimental/OnlineServices/Source/OnlineServicesInterface/Public/Online/OnlineId.h
@@ -9,20 +9,20 @@
 
 namespace UE::Online {
 
-/** Enum used as template argument to TOnlineIdHandle to make it a compile error to assign between id's of different types */
-enum class EOnlineIdType
+/** Tags used as template argument to TOnlineIdHandle to make it a compile error to assign between id's of different types */
+namespace OnlineIdHandleTags
 {
-	AccountId,
-	SessionId,
-	PartyId
-};
+	struct FAccount {};
+	struct FSession {};
+	struct FParty {};
+}
 
 /**
  * A handle to an id which uniquely identifies a persistent or transient online resource, i.e. account/session/party etc, within a given Online Services provider.
  * At most one id, and therefore one handle, exists for any given resource. The id and handle persist until the OnlineServices module is unloaded.
  * Passed to and returned from OnlineServices APIs.
  */ 
-template<EOnlineIdType IdType>
+template<typename IdType>
 class TOnlineIdHandle
 {
 public:
@@ -45,16 +45,16 @@ private:
 	uint32 Value = uint32(EOnlineServices::Null) << 24;
 };
 
-using FOnlineAccountIdHandle = TOnlineIdHandle<EOnlineIdType::AccountId>;
+using FOnlineAccountIdHandle = TOnlineIdHandle<OnlineIdHandleTags::FAccount>;
 
-template<EOnlineIdType IdType>
+template<typename IdType>
 inline uint32 GetTypeHash(const TOnlineIdHandle<IdType>& Handle)
 {
 	using ::GetTypeHash;
 	return HashCombine(GetTypeHash(Handle.GetOnlineServicesType()), GetTypeHash(Handle.GetHandle()));
 }
 
-template<EOnlineIdType IdType>
+template<typename IdType>
 inline void LexFromString(const TOnlineIdHandle<IdType>& Id, const TCHAR* String)
 {
 	// TODO: should instead just implement ParseOnlineExecParams
diff --git a/Engine/Plugins/Online/Experimental/OnlineServices/Source/OnlineServicesInterface/Public/Online/OnlineServices.h b/Engine/Plugins/Online/Experimental/OnlineServices/Source/OnlineServicesInterface/Public/Online/OnlineServices.h
--- a/Engine/Plugins/Online/Experimental/OnlineServices/Source/OnlineServicesInterface/Public/Online/OnlineServices.h
+++ b/Engine/Plugins/Online/Experimental/OnlineServices/Source/OnlineServicesInterface/Public/Online/OnlineServices.h
@@ -83,7 +83,7 @@ TSharedPtr<ServicesClass> GetServices(FName InstanceName = NAME_None)
  */
 ONLINESERVICESINTERFACE_API void DestroyServices(EOnlineServices OnlineServices = EOnlineServices::Default, FName InstanceName = NAME_None);
 
-template<EOnlineIdType IdType>
+template<typename IdType>
 inline FString ToLogString(const TOnlineIdHandle<IdType>& Id)
 {
 	FString Result;
diff --git a/Engine/Plugins/Online/Experimental/OnlineServicesEOS/Source/Private/OnlineIdEOS.h b/Engine/Plugins/Online/Experimental/OnlineServicesEOS/Source/Private/OnlineIdEOS.h
--- a/Engine/Plugins/Online/Experimental/OnlineServicesEOS/Source/Private/OnlineIdEOS.h
+++ b/Engine/Plugins/Online/Experimental/OnlineServicesEOS/Source/Private/OnlineIdEOS.h
@@ -58,7 +58,7 @@ FOnlineAccountIdHandle FindAccountIdChecked(const EOS_ProductUserId EpicAccountI
 
 FOnlineAccountIdHandle CreateAccountId(const EOS_EpicAccountId EpicAccountId, const EOS_ProductUserId ProductUserId);
 
-template<EOnlineIdType IdType>
+template<typename IdType>
 inline bool ValidateOnlineId(const TOnlineIdHandle<IdType>& Handle)
 {
 	return Handle.GetOnlineServicesType() == EOnlineServices::Epic && Handle.IsValid();
