Fixes for Horde issues.
Added ToLogString methods for uint8, int8, uint16 and int16 types.
Added null check to Log in AuthEOS

diff --git a/Engine/Plugins/Online/Experimental/OnlineServices/Source/OnlineServicesInterface/Public/Online/OnlineUtils.h b/Engine/Plugins/Online/Experimental/OnlineServices/Source/OnlineServicesInterface/Public/Online/OnlineUtils.h
--- a/Engine/Plugins/Online/Experimental/OnlineServices/Source/OnlineServicesInterface/Public/Online/OnlineUtils.h
+++ b/Engine/Plugins/Online/Experimental/OnlineServices/Source/OnlineServicesInterface/Public/Online/OnlineUtils.h
@@ -103,6 +103,26 @@ inline FString ToLogString(const FString& String)
 	return String;
 }
 
+inline FString ToLogString(uint8 Value)
+{
+	return FString::Printf(TEXT("%hhu"), Value);
+}
+
+inline FString ToLogString(int8 Value)
+{
+	return FString::Printf(TEXT("%hhi"), Value);
+}
+
+inline FString ToLogString(uint16 Value)
+{
+	return FString::Printf(TEXT("%hu"), Value);
+}
+
+inline FString ToLogString(int16 Value)
+{
+	return FString::Printf(TEXT("%hi"), Value);
+}
+
 inline FString ToLogString(uint32 Value)
 {
 	return FString::Printf(TEXT("%u"), Value);
diff --git a/Engine/Plugins/Online/Experimental/OnlineServicesEOS/Source/Private/AuthEOS.cpp b/Engine/Plugins/Online/Experimental/OnlineServicesEOS/Source/Private/AuthEOS.cpp
--- a/Engine/Plugins/Online/Experimental/OnlineServicesEOS/Source/Private/AuthEOS.cpp
+++ b/Engine/Plugins/Online/Experimental/OnlineServicesEOS/Source/Private/AuthEOS.cpp
@@ -386,7 +386,7 @@ TOnlineAsyncOpHandle<FAuthLogin> FAuthEOS::Login(FAuthLogin::Params&& Params)
 	})
 	.Then([this](TOnlineAsyncOp<FAuthLogin>& InAsyncOp, const EOS_Auth_LinkAccountCallbackInfo* Data)
 	{
-		UE_LOG(LogTemp, Verbose, TEXT("[FAuthEOS::Login] EOS_Auth_LinkAccount Result: [%s]"), *LexToString(Data->ResultCode));
+		UE_LOG(LogTemp, Verbose, TEXT("[FAuthEOS::Login] EOS_Auth_LinkAccount Result: [%s]"), (Data == nullptr) ? TEXT("Null") : *LexToString(Data->ResultCode));
 		
 		// If Data is NULL, it means the account was already linked
 		if (Data == nullptr || Data->ResultCode == EOS_EResult::EOS_Success)
