- Add registry of net id registries

diff --git a/Engine/Plugins/Online/Experimental/OnlineServices/Source/OnlineServicesInterface/Private/OnlineId.cpp b/Engine/Plugins/Online/Experimental/OnlineServices/Source/OnlineServicesInterface/Private/OnlineId.cpp
new file mode 100644
--- /dev/null
+++ b/Engine/Plugins/Online/Experimental/OnlineServices/Source/OnlineServicesInterface/Private/OnlineId.cpp
@@ -0,0 +1,58 @@
+// Copyright Epic Games, Inc. All Rights Reserved.
+
+#include "Online/OnlineId.h"
+
+#include "Misc/LazySingleton.h"
+
+namespace UE::Online {
+
+FOnlineIdRegistryRegistry& FOnlineIdRegistryRegistry::Get()
+{
+	return TLazySingleton<FOnlineIdRegistryRegistry>::Get();
+}
+
+void FOnlineIdRegistryRegistry::TearDown()
+{
+	return TLazySingleton<FOnlineIdRegistryRegistry>::TearDown();
+}
+
+void FOnlineIdRegistryRegistry::RegisterAccountIdRegistry(EOnlineServices OnlineServices, IOnlineAccountIdRegistry* Registry, int32 Priority)
+{
+	FAccountIdRegistryAndPriority* Found = AccountIdRegistries.Find(OnlineServices);
+	if (Found == nullptr || Found->Priority < Priority)
+	{
+		AccountIdRegistries.Emplace(OnlineServices, FAccountIdRegistryAndPriority(Registry, Priority));
+	}
+}
+
+void FOnlineIdRegistryRegistry::UnregisterAccountIdRegistry(EOnlineServices OnlineServices, int32 Priority)
+{
+	FAccountIdRegistryAndPriority* Found = AccountIdRegistries.Find(OnlineServices);
+	if (Found != nullptr && Found->Priority == Priority)
+	{
+		AccountIdRegistries.Remove(OnlineServices);
+	}
+}
+
+IOnlineAccountIdRegistry* FOnlineIdRegistryRegistry::GetAccountIdRegistry(EOnlineServices OnlineServices)
+{
+	IOnlineAccountIdRegistry* Registry = nullptr;
+	if (FAccountIdRegistryAndPriority* Found = AccountIdRegistries.Find(OnlineServices))
+	{
+		Registry = Found->Registry;
+	}
+	return Registry;
+}
+
+FString ToLogString(const FOnlineAccountIdHandle& Id)
+{
+	FString Result;
+	if (IOnlineAccountIdRegistry* Registry = FOnlineIdRegistryRegistry::Get().GetAccountIdRegistry(Id.GetOnlineServicesType()))
+	{
+		Result = Registry->ToLogString(Id);
+	}
+	return Result;
+}
+
+
+/* UE::Online */ }
diff --git a/Engine/Plugins/Online/Experimental/OnlineServices/Source/OnlineServicesInterface/Public/Online/OnlineId.h b/Engine/Plugins/Online/Experimental/OnlineServices/Source/OnlineServicesInterface/Public/Online/OnlineId.h
--- a/Engine/Plugins/Online/Experimental/OnlineServices/Source/OnlineServicesInterface/Public/Online/OnlineId.h
+++ b/Engine/Plugins/Online/Experimental/OnlineServices/Source/OnlineServicesInterface/Public/Online/OnlineId.h
@@ -60,4 +60,72 @@ inline void LexFromString(const TOnlineIdHandle<IdType>& Id, const TCHAR* String
 	// TODO: should instead just implement ParseOnlineExecParams
 }
 
+ONLINESERVICESINTERFACE_API FString ToLogString(const FOnlineAccountIdHandle& Id);
+
+template<typename IdType>
+class IOnlineIdRegistry
+{
+public:
+	virtual FString ToLogString(const TOnlineIdHandle<IdType>& Handle) const = 0;
+};
+
+using IOnlineAccountIdRegistry = IOnlineIdRegistry<OnlineIdHandleTags::FAccount>;
+
+class FOnlineIdRegistryRegistry
+{
+public:
+	/**
+	 * Get the FOnlineIdRegistryRegistry singleton
+	 *
+	 * @return The FOnlineIdRegistryRegistry singleton instance
+	 */
+	ONLINESERVICESINTERFACE_API static FOnlineIdRegistryRegistry& Get();
+
+	/**
+	 * Tear down the singleton instance
+	 */
+	ONLINESERVICESINTERFACE_API static void TearDown();
+
+	/**
+	 * Register a registry for a given OnlineServices implementation and TOnlineIdHandle type
+	 *
+	 * @param OnlineServices Services that the registry is for
+	 * @param Registry the registry of online ids
+	 * @param Priority Integer priority, allows an existing registry to be extended and registered with a higher priority so it is used instead
+	 */
+	ONLINESERVICESINTERFACE_API void RegisterAccountIdRegistry(EOnlineServices OnlineServices, IOnlineAccountIdRegistry* Registry, int32 Priority = 0);
+
+	/**
+	* Unregister a previously registered Id registry
+	*
+	* @param OnlineServices Services that the registry is for
+	* @param Priority Integer priority, will be unregistered only if the priority matches the one that is registered
+	*/
+	ONLINESERVICESINTERFACE_API void UnregisterAccountIdRegistry(EOnlineServices OnlineServices, int32 Priority = 0);
+
+	/**
+	 * Get the account id registry
+	 *
+	 * @param OnlineServices Type of online services for the IOnlineServices instance
+	 *
+	 * @return The account id registry, or an invalid pointer if it is not registered
+	 */
+	ONLINESERVICESINTERFACE_API IOnlineAccountIdRegistry* GetAccountIdRegistry(EOnlineServices OnlineServices);
+
+private:
+	struct FAccountIdRegistryAndPriority
+	{
+		FAccountIdRegistryAndPriority(IOnlineAccountIdRegistry* InRegistry, int32 InPriority)
+			: Registry(InRegistry), Priority(InPriority) {}
+
+		IOnlineAccountIdRegistry* Registry;
+		int32 Priority;
+	};
+
+	TMap<EOnlineServices, FAccountIdRegistryAndPriority> AccountIdRegistries;
+
+	FOnlineIdRegistryRegistry() = default;
+	~FOnlineIdRegistryRegistry() = default;
+	friend FLazySingleton;
+};
 /* UE::Online */ }
diff --git a/Engine/Plugins/Online/Experimental/OnlineServices/Source/OnlineServicesInterface/Public/Online/OnlineServices.h b/Engine/Plugins/Online/Experimental/OnlineServices/Source/OnlineServicesInterface/Public/Online/OnlineServices.h
--- a/Engine/Plugins/Online/Experimental/OnlineServices/Source/OnlineServicesInterface/Public/Online/OnlineServices.h
+++ b/Engine/Plugins/Online/Experimental/OnlineServices/Source/OnlineServicesInterface/Public/Online/OnlineServices.h
@@ -47,11 +47,6 @@ public:
 	 *
 	 */
 	virtual IExternalUIPtr GetExternalUIInterface() = 0;
-
-	/**
-	 * 
-	 */
-	virtual FString ToLogString(const FOnlineAccountIdHandle& Handle) = 0;
 };
 
 /**
@@ -83,15 +78,4 @@ TSharedPtr<ServicesClass> GetServices(FName InstanceName = NAME_None)
  */
 ONLINESERVICESINTERFACE_API void DestroyServices(EOnlineServices OnlineServices = EOnlineServices::Default, FName InstanceName = NAME_None);
 
-template<typename IdType>
-inline FString ToLogString(const TOnlineIdHandle<IdType>& Id)
-{
-	FString Result;
-	if (TSharedPtr<IOnlineServices> Services = GetServices(Id.GetOnlineServicesType()))
-	{
-		Result = Services->ToLogString(Id);
-	}
-	return Result;
-}
-
 /* UE::Online */ }
diff --git a/Engine/Plugins/Online/Experimental/OnlineServicesEOS/Source/Private/OnlineIdEOS.cpp b/Engine/Plugins/Online/Experimental/OnlineServicesEOS/Source/Private/OnlineIdEOS.cpp
--- a/Engine/Plugins/Online/Experimental/OnlineServicesEOS/Source/Private/OnlineIdEOS.cpp
+++ b/Engine/Plugins/Online/Experimental/OnlineServicesEOS/Source/Private/OnlineIdEOS.cpp
@@ -146,6 +146,22 @@ FOnlineAccountIdDataEOS FOnlineAccountIdRegistryEOS::GetIdData(const FOnlineAcco
 	return FOnlineAccountIdDataEOS();
 }
 
+FString FOnlineAccountIdRegistryEOS::ToLogString(const FOnlineAccountIdHandle& Handle) const
+{
+	FString Result;
+	if (ValidateOnlineId(Handle))
+	{
+		const FOnlineAccountIdDataEOS& IdData = GetIdData(Handle);
+		Result = FString::Printf(TEXT("EAS=[%s] EOS=[%s]"), *LexToString(IdData.EpicAccountId), *LexToString(IdData.ProductUserId));
+	}
+	else
+	{
+		check(!Handle.IsValid()); // Check we haven't been passed a valid handle for a different EOnlineServices.
+		Result = TEXT("Invalid");
+	}
+	return Result;
+}
+
 EOS_EpicAccountId GetEpicAccountId(const FOnlineAccountIdHandle& Handle)
 {
 	return FOnlineAccountIdRegistryEOS::Get().GetIdData(Handle).EpicAccountId;
diff --git a/Engine/Plugins/Online/Experimental/OnlineServicesEOS/Source/Private/OnlineIdEOS.h b/Engine/Plugins/Online/Experimental/OnlineServicesEOS/Source/Private/OnlineIdEOS.h
--- a/Engine/Plugins/Online/Experimental/OnlineServicesEOS/Source/Private/OnlineIdEOS.h
+++ b/Engine/Plugins/Online/Experimental/OnlineServicesEOS/Source/Private/OnlineIdEOS.h
@@ -25,7 +25,7 @@ public:
 /**
  * Account id registry specifically for EOS id's which are segmented.
  */
-class FOnlineAccountIdRegistryEOS
+class FOnlineAccountIdRegistryEOS : public IOnlineAccountIdRegistry
 {
 public:
 	static FOnlineAccountIdRegistryEOS& Get();
@@ -37,6 +37,10 @@ public:
 	// Return copies as it is not thread safe to return pointers/references to array elements, in case the array is grown+relocated on another thread.
 	FOnlineAccountIdDataEOS GetIdData(const FOnlineAccountIdHandle& Handle) const;
 
+	// Begin IOnlineAccountIdRegistry
+	virtual FString ToLogString(const FOnlineAccountIdHandle& Handle) const override;
+	// End IOnlineAccountIdRegistry
+
 private:
 	mutable FRWLock Lock;
 
@@ -44,6 +48,8 @@ private:
 
 	TMap<EOS_EpicAccountId, FOnlineAccountIdHandle> EpicAccountIdToHandle; // Map of EOS_EpicAccountId to the associated handle.
 	TMap<EOS_ProductUserId, FOnlineAccountIdHandle> ProductUserIdToHandle; // Map of EOS_ProductUserId to the associated handle.
+
+	virtual ~FOnlineAccountIdRegistryEOS() = default;
 };
 
 EOS_EpicAccountId GetEpicAccountId(const FOnlineAccountIdHandle& Handle);
diff --git a/Engine/Plugins/Online/Experimental/OnlineServicesEOS/Source/Private/OnlineServicesEOS.cpp b/Engine/Plugins/Online/Experimental/OnlineServicesEOS/Source/Private/OnlineServicesEOS.cpp
--- a/Engine/Plugins/Online/Experimental/OnlineServicesEOS/Source/Private/OnlineServicesEOS.cpp
+++ b/Engine/Plugins/Online/Experimental/OnlineServicesEOS/Source/Private/OnlineServicesEOS.cpp
@@ -88,21 +88,6 @@ void FOnlineServicesEOS::Initialize()
 	FOnlineServicesCommon::Initialize();
 }
 
-FString FOnlineServicesEOS::ToLogString(const FOnlineAccountIdHandle& Handle)
-{
-	FString Result;
-	if (Handle.IsValid())
-	{
-		const FOnlineAccountIdDataEOS& AccountIdData = FOnlineAccountIdRegistryEOS::Get().GetIdData(Handle);
-		Result = FString::Printf(TEXT("EAS=[%s] EOS=[%s]"), *LexToString(AccountIdData.EpicAccountId), *LexToString(AccountIdData.ProductUserId));
-	}
-	else
-	{
-		Result = TEXT("Invalid");
-	}
-	return Result;
-}
-
 EOS_HPlatform FOnlineServicesEOS::GetEOSPlatformHandle() const
 {
 	if (EOSPlatformHandle)
diff --git a/Engine/Plugins/Online/Experimental/OnlineServicesEOS/Source/Private/OnlineServicesEOSModule.cpp b/Engine/Plugins/Online/Experimental/OnlineServicesEOS/Source/Private/OnlineServicesEOSModule.cpp
--- a/Engine/Plugins/Online/Experimental/OnlineServicesEOS/Source/Private/OnlineServicesEOSModule.cpp
+++ b/Engine/Plugins/Online/Experimental/OnlineServicesEOS/Source/Private/OnlineServicesEOSModule.cpp
@@ -6,6 +6,7 @@
 
 #include "Online/OnlineServicesRegistry.h"
 #include "OnlineServicesEOS.h"
+#include "OnlineIdEOS.h"
 
 namespace UE::Online
 {
@@ -31,6 +32,7 @@ protected:
 void FOnlineServicesEOSModule::StartupModule()
 {
 	FOnlineServicesRegistry::Get().RegisterServicesFactory(EOnlineServices::Epic, MakeUnique<FOnlineServicesFactoryEOS>());
+	FOnlineIdRegistryRegistry::Get().RegisterAccountIdRegistry(EOnlineServices::Epic, &FOnlineAccountIdRegistryEOS::Get());
 }
 
 /* UE::Online */ }
diff --git a/Engine/Plugins/Online/Experimental/OnlineServicesEOS/Source/Public/OnlineServicesEOS.h b/Engine/Plugins/Online/Experimental/OnlineServicesEOS/Source/Public/OnlineServicesEOS.h
--- a/Engine/Plugins/Online/Experimental/OnlineServicesEOS/Source/Public/OnlineServicesEOS.h
+++ b/Engine/Plugins/Online/Experimental/OnlineServicesEOS/Source/Public/OnlineServicesEOS.h
@@ -25,8 +25,6 @@ public:
 	static EOnlineServices GetServicesProvider() { return EOnlineServices::Epic; }
 
 	EOS_HPlatform GetEOSPlatformHandle() const;
-
-	virtual FString ToLogString(const FOnlineAccountIdHandle& Handle) override;
 protected:
 	IEOSPlatformHandlePtr EOSPlatformHandle;
 };
