- Fix horde errors

diff --git a/Engine/Plugins/Online/Experimental/OnlineServices/Source/OnlineServicesInterface/Public/Online/OnlineId.h b/Engine/Plugins/Online/Experimental/OnlineServices/Source/OnlineServicesInterface/Public/Online/OnlineId.h
--- a/Engine/Plugins/Online/Experimental/OnlineServices/Source/OnlineServicesInterface/Public/Online/OnlineId.h
+++ b/Engine/Plugins/Online/Experimental/OnlineServices/Source/OnlineServicesInterface/Public/Online/OnlineId.h
@@ -126,6 +126,6 @@ private:
 
 	FOnlineIdRegistryRegistry() = default;
 	~FOnlineIdRegistryRegistry() = default;
-	friend FLazySingleton;
+	friend class FLazySingleton;
 };
 /* UE::Online */ }
