- Fix horde errors

diff --git a/Engine/Plugins/Online/Experimental/OnlineServices/Source/OnlineServicesInterface/Public/Online/OnlineId.h b/Engine/Plugins/Online/Experimental/OnlineServices/Source/OnlineServicesInterface/Public/Online/OnlineId.h
--- a/Engine/Plugins/Online/Experimental/OnlineServices/Source/OnlineServicesInterface/Public/Online/OnlineId.h
+++ b/Engine/Plugins/Online/Experimental/OnlineServices/Source/OnlineServicesInterface/Public/Online/OnlineId.h
@@ -112,6 +112,9 @@ public:
 	 */
 	ONLINESERVICESINTERFACE_API IOnlineAccountIdRegistry* GetAccountIdRegistry(EOnlineServices OnlineServices);
 
+	~FOnlineIdRegistryRegistry() = default;
+	FOnlineIdRegistryRegistry() = default;
+
 private:
 	struct FAccountIdRegistryAndPriority
 	{
@@ -124,8 +127,6 @@ private:
 
 	TMap<EOnlineServices, FAccountIdRegistryAndPriority> AccountIdRegistries;
 
-	FOnlineIdRegistryRegistry() = default;
-	~FOnlineIdRegistryRegistry() = default;
 	friend class FLazySingleton;
 };
 /* UE::Online */ }
