OnlineServicesCommon - Fix per-interface config loading

diff --git a/Engine/Plugins/Online/Experimental/OnlineServices/Source/OnlineServicesCommon/Public/Online/OnlineComponent.h b/Engine/Plugins/Online/Experimental/OnlineServices/Source/OnlineServicesCommon/Public/Online/OnlineComponent.h
--- a/Engine/Plugins/Online/Experimental/OnlineServices/Source/OnlineServicesCommon/Public/Online/OnlineComponent.h
+++ b/Engine/Plugins/Online/Experimental/OnlineServices/Source/OnlineServicesCommon/Public/Online/OnlineComponent.h
@@ -42,7 +42,7 @@ public:
 	template <typename StructType>
 	bool LoadConfig(StructType& Struct, const FString& OperationName = FString()) const
 	{
-		return static_cast<const ComponentType*>(this)->GetServices()->LoadConfig(Struct, GetConfigName(), OperationName);
+		return GetServices().LoadConfig(Struct, GetConfigName(), OperationName);
 	}
 
 	const FString& GetConfigName() const { return InterfaceName; }
@@ -71,6 +71,12 @@ public:
 		return static_cast<ServicesType&>(Services);
 	}
 
+	template <typename ServicesType = FOnlineServicesCommon>
+	const ServicesType& GetServices() const
+	{
+		return static_cast<const ServicesType&>(Services);
+	}
+
 	/* Queue for executing tasks in serial */
 	FOnlineAsyncOpQueue& GetSerialQueue()
 	{
diff --git a/Engine/Plugins/Online/Experimental/OnlineServices/Source/OnlineServicesCommon/Public/Online/OnlineServicesCommon.h b/Engine/Plugins/Online/Experimental/OnlineServices/Source/OnlineServicesCommon/Public/Online/OnlineServicesCommon.h
--- a/Engine/Plugins/Online/Experimental/OnlineServices/Source/OnlineServicesCommon/Public/Online/OnlineServicesCommon.h
+++ b/Engine/Plugins/Online/Experimental/OnlineServices/Source/OnlineServicesCommon/Public/Online/OnlineServicesCommon.h
@@ -161,7 +161,7 @@ public:
 	 */
 	const FString& GetConfigName() const { return ConfigName; }
 
-	TArray<FString> GetConfigSectionHeiarchy(const FString& OperationName = FString())
+	TArray<FString> GetConfigSectionHeiarchy(const FString& OperationName = FString()) const
 	{
 		TArray<FString> SectionHeiarchy;
 		FString SectionName = TEXT("OnlineServices");
@@ -191,7 +191,7 @@ public:
 	 * @return true if a value was loaded
 	 */
 	template <typename StructType>
-	bool LoadConfig(StructType& Struct, const FString& InterfaceName = FString(), const FString& OperationName = FString())
+	bool LoadConfig(StructType& Struct, const FString& InterfaceName = FString(), const FString& OperationName = FString()) const
 	{
 		TArray<FString> SectionHeiarchy;
 		FString SectionName = TEXT("OnlineServices");
@@ -218,7 +218,7 @@ public:
 	 * 
 	 * @return Array of the sections with overrides for values to be loaded from
 	 */
-	TArray<FString> GetConfigSectionHeirachWithOverrides(const TArray<FString>& SectionHeiarchy)
+	TArray<FString> GetConfigSectionHeirachWithOverrides(const TArray<FString>& SectionHeiarchy) const
 	{
 		TArray<FString> SectionHeiarchyWithOverrides;
 		for (const FString& Section : SectionHeiarchy)
@@ -243,17 +243,11 @@ public:
 	 * @return true if a value was loaded
 	 */
 	template <typename StructType>
-	bool LoadConfig(StructType& Struct, const TArray<FString>& SectionHeiarchy)
+	bool LoadConfig(StructType& Struct, const TArray<FString>& SectionHeiarchy) const
 	{
 		return ::UE::Online::LoadConfig(*ConfigProvider, GetConfigSectionHeirachWithOverrides(SectionHeiarchy), Struct);
 	}
 
-	template <typename OpType>
-	bool LoadOperationConfig(FOperationConfig& OutConfig, const FString& InterfaceName = FString(), const FString& OperationName = FString())
-	{
-		return false;
-	}
-
 	/* Get op (OnlineServices) */
 	template <typename OpType>
 	TOnlineAsyncOpRef<OpType> GetOp(typename OpType::Params&& Params)
