Add PIE support for online services (OSSv2)
OnlineServices changes

diff --git a/Engine/Plugins/Online/Experimental/OnlineServices/Source/OnlineServicesCommon/Private/Online/ExternalUICommon.cpp b/Engine/Plugins/Online/Experimental/OnlineServices/Source/OnlineServicesCommon/Private/Online/ExternalUICommon.cpp
--- a/Engine/Plugins/Online/Experimental/OnlineServices/Source/OnlineServicesCommon/Private/Online/ExternalUICommon.cpp
+++ b/Engine/Plugins/Online/Experimental/OnlineServices/Source/OnlineServicesCommon/Private/Online/ExternalUICommon.cpp
@@ -32,4 +32,9 @@ TOnlineAsyncOpHandle<FExternalUIShowFriendsUI> FExternalUICommon::ShowFriendsUI(
 	return Operation->GetHandle();
 }
 
+TOnlineEvent<void(const FExternalUIStatusChanged&)> FExternalUICommon::OnExternalUIStatusChanged()
+{
+	return OnExternalUIStatusChangedEvent;
+}
+
 /* UE::Online */ }
diff --git a/Engine/Plugins/Online/Experimental/OnlineServices/Source/OnlineServicesCommon/Public/Online/ExternalUICommon.h b/Engine/Plugins/Online/Experimental/OnlineServices/Source/OnlineServicesCommon/Public/Online/ExternalUICommon.h
--- a/Engine/Plugins/Online/Experimental/OnlineServices/Source/OnlineServicesCommon/Public/Online/ExternalUICommon.h
+++ b/Engine/Plugins/Online/Experimental/OnlineServices/Source/OnlineServicesCommon/Public/Online/ExternalUICommon.h
@@ -22,6 +22,10 @@ public:
 	// IExternalUI
 	virtual TOnlineAsyncOpHandle<FExternalUIShowLoginUI> ShowLoginUI(FExternalUIShowLoginUI::Params&& Params) override;
 	virtual TOnlineAsyncOpHandle<FExternalUIShowFriendsUI> ShowFriendsUI(FExternalUIShowFriendsUI::Params&& Params) override;
+
+	virtual TOnlineEvent<void(const FExternalUIStatusChanged&)> OnExternalUIStatusChanged() override;
+protected:
+	TOnlineEventCallable<void(const FExternalUIStatusChanged&)> OnExternalUIStatusChangedEvent;
 };
 
 /* UE::Online */ }
\ No newline at end of file
diff --git a/Engine/Plugins/Online/Experimental/OnlineServices/Source/OnlineServicesInterface/Private/OnlineServices.cpp b/Engine/Plugins/Online/Experimental/OnlineServices/Source/OnlineServicesInterface/Private/OnlineServices.cpp
--- a/Engine/Plugins/Online/Experimental/OnlineServices/Source/OnlineServicesInterface/Private/OnlineServices.cpp
+++ b/Engine/Plugins/Online/Experimental/OnlineServices/Source/OnlineServicesInterface/Private/OnlineServices.cpp
@@ -6,6 +6,11 @@
 
 namespace UE::Online {
 
+bool IsLoaded(EOnlineServices OnlineServices, FName InstanceName)
+{
+	return FOnlineServicesRegistry::Get().IsLoaded(OnlineServices, InstanceName);
+}
+
 TSharedPtr<IOnlineServices> GetServices(EOnlineServices OnlineServices, FName InstanceName)
 {
 	return FOnlineServicesRegistry::Get().GetNamedServicesInstance(OnlineServices, InstanceName);
diff --git a/Engine/Plugins/Online/Experimental/OnlineServices/Source/OnlineServicesInterface/Private/OnlineServicesDelegates.cpp b/Engine/Plugins/Online/Experimental/OnlineServices/Source/OnlineServicesInterface/Private/OnlineServicesDelegates.cpp
new file mode 100644
--- /dev/null
+++ b/Engine/Plugins/Online/Experimental/OnlineServices/Source/OnlineServicesInterface/Private/OnlineServicesDelegates.cpp
@@ -0,0 +1,9 @@
+// Copyright Epic Games, Inc. All Rights Reserved.
+
+#include "Online/OnlineServicesDelegates.h"
+
+namespace UE::Online {
+
+FOnOnlineServicesCreated OnOnlineServicesCreated;
+
+/* UE::Online */ }
diff --git a/Engine/Plugins/Online/Experimental/OnlineServices/Source/OnlineServicesInterface/Private/OnlineServicesRegistry.cpp b/Engine/Plugins/Online/Experimental/OnlineServices/Source/OnlineServicesInterface/Private/OnlineServicesRegistry.cpp
--- a/Engine/Plugins/Online/Experimental/OnlineServices/Source/OnlineServicesInterface/Private/OnlineServicesRegistry.cpp
+++ b/Engine/Plugins/Online/Experimental/OnlineServices/Source/OnlineServicesInterface/Private/OnlineServicesRegistry.cpp
@@ -1,6 +1,7 @@
 // Copyright Epic Games, Inc. All Rights Reserved.
 
 #include "Online/OnlineServicesRegistry.h"
+#include "Online/OnlineServicesDelegates.h"
 
 #include "Misc/LazySingleton.h"
 
@@ -34,6 +35,16 @@ void FOnlineServicesRegistry::UnregisterServicesFactory(EOnlineServices OnlineSe
 	}
 }
 
+bool FOnlineServicesRegistry::IsLoaded(EOnlineServices OnlineServices, FName InstanceName) const
+{
+	bool bExists = false;
+	if (const TMap<FName, TSharedRef<IOnlineServices>>* OnlineServicesInstances = NamedServiceInstances.Find(OnlineServices))
+	{
+		bExists = OnlineServicesInstances->Find(InstanceName) != nullptr;
+	}
+	return bExists;
+}
+
 TSharedPtr<IOnlineServices> FOnlineServicesRegistry::GetNamedServicesInstance(EOnlineServices OnlineServices, FName InstanceName)
 {
 	TSharedPtr<IOnlineServices> Services;
@@ -48,6 +59,7 @@ TSharedPtr<IOnlineServices> FOnlineServicesRegistry::GetNamedServicesInstance(EO
 		if (Services.IsValid())
 		{
 			NamedServiceInstances.FindOrAdd(OnlineServices).Add(InstanceName, Services.ToSharedRef());
+			OnOnlineServicesCreated.Broadcast(Services.ToSharedRef());
 		}
 	}
 
@@ -78,6 +90,17 @@ TSharedPtr<IOnlineServices> FOnlineServicesRegistry::CreateServices(EOnlineServi
 	return Services;
 }
 
+void FOnlineServicesRegistry::GetAllServicesInstances(TArray<TSharedRef<IOnlineServices>>& OutOnlineServices) const
+{
+	for (const TPair<EOnlineServices, TMap<FName, TSharedRef<IOnlineServices>>>& OnlineServiceTypesMaps : NamedServiceInstances)
+	{
+		for (const TPair<FName, TSharedRef<IOnlineServices>>& NamedInstance : OnlineServiceTypesMaps.Value)
+		{
+			OutOnlineServices.Emplace(NamedInstance.Value);
+		}
+	}
+}
+
 FOnlineServicesRegistry::~FOnlineServicesRegistry()
 {
 	for (TPair<EOnlineServices, TMap<FName, TSharedRef<IOnlineServices>>>& ServiceInstances : NamedServiceInstances)
diff --git a/Engine/Plugins/Online/Experimental/OnlineServices/Source/OnlineServicesInterface/Public/Online/ExternalUI.h b/Engine/Plugins/Online/Experimental/OnlineServices/Source/OnlineServicesInterface/Public/Online/ExternalUI.h
--- a/Engine/Plugins/Online/Experimental/OnlineServices/Source/OnlineServicesInterface/Public/Online/ExternalUI.h
+++ b/Engine/Plugins/Online/Experimental/OnlineServices/Source/OnlineServicesInterface/Public/Online/ExternalUI.h
@@ -45,6 +45,13 @@ struct FExternalUIShowFriendsUI
 	};
 };
 
+/** Struct for ExternalUIStatusChanged event */
+struct FExternalUIStatusChanged
+{
+	/** Whether the external UI is being opened or not */
+	bool bIsOpening = false;
+};
+
 class IExternalUI
 {
 public:
@@ -57,6 +64,11 @@ public:
 	 * Shows the online service's friends UI
 	 */
 	virtual TOnlineAsyncOpHandle<FExternalUIShowFriendsUI> ShowFriendsUI(FExternalUIShowFriendsUI::Params&& Params) = 0;
+
+	/**
+	 * Event triggered when an external UI's status changes
+	 */
+	virtual TOnlineEvent<void(const FExternalUIStatusChanged&)> OnExternalUIStatusChanged() = 0;
 };
 
 namespace Meta {
@@ -78,6 +90,10 @@ END_ONLINE_STRUCT_META()
 BEGIN_ONLINE_STRUCT_META(FExternalUIShowFriendsUI::Result)
 END_ONLINE_STRUCT_META()
 
+BEGIN_ONLINE_STRUCT_META(FExternalUIStatusChanged)
+	ONLINE_STRUCT_FIELD(FExternalUIStatusChanged, bIsOpening)
+END_ONLINE_STRUCT_META()
+
 /* Meta*/ }
 
 /* UE::Online */ }
diff --git a/Engine/Plugins/Online/Experimental/OnlineServices/Source/OnlineServicesInterface/Public/Online/OnlineServices.h b/Engine/Plugins/Online/Experimental/OnlineServices/Source/OnlineServicesInterface/Public/Online/OnlineServices.h
--- a/Engine/Plugins/Online/Experimental/OnlineServices/Source/OnlineServicesInterface/Public/Online/OnlineServices.h
+++ b/Engine/Plugins/Online/Experimental/OnlineServices/Source/OnlineServicesInterface/Public/Online/OnlineServices.h
@@ -51,7 +51,16 @@ public:
 };
 
 /**
- * Get an instance of the online subsystem
+ * Check if an instance of the online service is loaded
+ *
+ * @param OnlineServices Type of online services to retrieve
+ * @param InstanceName Name of the services instance to retrieve
+ * @return The services instance or an invalid pointer if the services is unavailable
+ */
+ONLINESERVICESINTERFACE_API bool IsLoaded(EOnlineServices OnlineServices = EOnlineServices::Default, FName InstanceName = NAME_None);
+
+/**
+ * Get an instance of the online service
  *
  * @param OnlineServices Type of online services to retrieve
  * @param InstanceName Name of the services instance to retrieve
@@ -72,7 +81,7 @@ TSharedPtr<ServicesClass> GetServices(FName InstanceName = NAME_None)
 }
 
 /**
- * Destroy an instance of the online subsystem
+ * Destroy an instance of the online service
  *
  * @param OnlineServices Type of online services to destroy
  * @param InstanceName Name of the services instance to destroy
diff --git a/Engine/Plugins/Online/Experimental/OnlineServices/Source/OnlineServicesInterface/Public/Online/OnlineServicesDelegates.h b/Engine/Plugins/Online/Experimental/OnlineServices/Source/OnlineServicesInterface/Public/Online/OnlineServicesDelegates.h
new file mode 100644
--- /dev/null
+++ b/Engine/Plugins/Online/Experimental/OnlineServices/Source/OnlineServicesInterface/Public/Online/OnlineServicesDelegates.h
@@ -0,0 +1,21 @@
+// Copyright Epic Games, Inc. All Rights Reserved.
+
+#pragma once
+
+#include "CoreMinimal.h"
+
+namespace UE::Online {
+
+/**
+ * Online services delegates that are more external to the online services themselves
+ */
+
+/**
+ * Notification that a new online subsystem instance has been created
+ *
+ * @param NewSubsystem the new instance created
+ */
+DECLARE_MULTICAST_DELEGATE_OneParam(FOnOnlineServicesCreated, TSharedRef<class IOnlineServices> /*NewServices*/);
+extern ONLINESERVICESINTERFACE_API FOnOnlineServicesCreated OnOnlineServicesCreated;
+
+/* UE::Online */ }
diff --git a/Engine/Plugins/Online/Experimental/OnlineServices/Source/OnlineServicesInterface/Public/Online/OnlineServicesRegistry.h b/Engine/Plugins/Online/Experimental/OnlineServices/Source/OnlineServicesInterface/Public/Online/OnlineServicesRegistry.h
--- a/Engine/Plugins/Online/Experimental/OnlineServices/Source/OnlineServicesInterface/Public/Online/OnlineServicesRegistry.h
+++ b/Engine/Plugins/Online/Experimental/OnlineServices/Source/OnlineServicesInterface/Public/Online/OnlineServicesRegistry.h
@@ -62,6 +62,16 @@ public:
 	 */
 	ONLINESERVICESINTERFACE_API void UnregisterServicesFactory(EOnlineServices OnlineServices, int32 Priority = 0);
 
+	/**
+	 * Check if an online service instance is loaded
+	 *
+	 * @param OnlineServices Type of online services for the IOnlineServices instance
+	 * @param InstanceName Name of the instance
+	 *
+	 * @return true if the instance is loaded
+	 */
+	ONLINESERVICESINTERFACE_API bool IsLoaded(EOnlineServices OnlineServices, FName InstanceName) const;
+
 	/**
 	 * Get a named instance of a specific IOnlineServices
 	 * 
@@ -89,6 +99,12 @@ public:
 	 */
 	ONLINESERVICESINTERFACE_API TSharedPtr<IOnlineServices> CreateServices(EOnlineServices OnlineServices);
 
+	/**
+	 * Get list of all instantiated OnlineServices
+	 * 
+	 * @param OutOnlineServices Array of online services to fill
+	 */
+	ONLINESERVICESINTERFACE_API void GetAllServicesInstances(TArray<TSharedRef<IOnlineServices>>& OutOnlineServices) const;
 private:
 	struct FFactoryAndPriority
 	{
