Fix compile error

diff --git a/Engine/Plugins/Online/Experimental/OnlineServices/Source/OnlineServicesCommon/Private/Online/LobbiesCommon.cpp b/Engine/Plugins/Online/Experimental/OnlineServices/Source/OnlineServicesCommon/Private/Online/LobbiesCommon.cpp
--- a/Engine/Plugins/Online/Experimental/OnlineServices/Source/OnlineServicesCommon/Private/Online/LobbiesCommon.cpp
+++ b/Engine/Plugins/Online/Experimental/OnlineServices/Source/OnlineServicesCommon/Private/Online/LobbiesCommon.cpp
@@ -167,7 +167,9 @@ void FLobbiesCommon::RegisterCommands()
 	RegisterCommand(&FLobbiesCommon::GetJoinedLobbies);
 	RegisterCommand(&FLobbiesCommon::GetReceivedInvitations);
 
+#if LOBBIES_FUNCTIONAL_TEST_ENABLED
 	RegisterCommand(&FLobbiesCommon::FunctionalTest);
+#endif
 }
 
 TOnlineAsyncOpHandle<FCreateLobby> FLobbiesCommon::CreateLobby(FCreateLobby::Params&& Params)
