Fix link issue

diff --git a/Engine/Plugins/Online/Experimental/OnlineServices/Source/OnlineServicesCommon/Private/Online/LobbiesCommon.cpp b/Engine/Plugins/Online/Experimental/OnlineServices/Source/OnlineServicesCommon/Private/Online/LobbiesCommon.cpp
--- a/Engine/Plugins/Online/Experimental/OnlineServices/Source/OnlineServicesCommon/Private/Online/LobbiesCommon.cpp
+++ b/Engine/Plugins/Online/Experimental/OnlineServices/Source/OnlineServicesCommon/Private/Online/LobbiesCommon.cpp
@@ -5,8 +5,6 @@
 #include "Online/Auth.h"
 #include "Online/OnlineServicesCommon.h"
 
-DEFINE_LOG_CATEGORY(LogLobbies);
-
 namespace UE::Online {
 
 template <typename DataType, typename OpType>
@@ -884,12 +882,12 @@ TOnlineAsyncOpHandle<FFunctionalTestLobbies> FLobbiesCommon::FunctionalTest(FFun
 
 TFuture<TOnlineResult<FLobbiesCommon::FFunctionalTestLoginUser>> FLobbiesCommon::FunctionalTestLoginUser(FFunctionalTestLoginUser::Params&& Params)
 {
-	UE_LOG(LogLobbies, Log, TEXT("[FLobbiesCommon::FunctionalTestLoginUser] Logging in user. PlatformUserId: %d"), FPlatformMisc::GetUserIndexForPlatformUser(Params.PlatformUserId));
+	UE_LOG(LogTemp, Log, TEXT("[FLobbiesCommon::FunctionalTestLoginUser] Logging in user. PlatformUserId: %d"), FPlatformMisc::GetUserIndexForPlatformUser(Params.PlatformUserId));
 
 	IAuth* AuthInterface = Services.Get<IAuth>();
 	if (!AuthInterface)
 	{
-		UE_LOG(LogLobbies, Warning, TEXT("[FLobbiesCommon::FunctionalTestLoginUser] Login failed. PlatformUserId: %d"), FPlatformMisc::GetUserIndexForPlatformUser(Params.PlatformUserId));
+		UE_LOG(LogTemp, Warning, TEXT("[FLobbiesCommon::FunctionalTestLoginUser] Login failed. PlatformUserId: %d"), FPlatformMisc::GetUserIndexForPlatformUser(Params.PlatformUserId));
 		return MakeFulfilledPromise<TOnlineResult<FFunctionalTestLoginUser>>(Errors::MissingInterface()).GetFuture();
 	}
 
@@ -907,12 +905,12 @@ TFuture<TOnlineResult<FLobbiesCommon::FFunctionalTestLoginUser>> FLobbiesCommon:
 	{
 		if (LoginResult.IsError())
 		{
-			UE_LOG(LogLobbies, Warning, TEXT("[FLobbiesCommon::FunctionalTestLoginUser] Login failed. PlatformUserId: %d, Result: %s"), FPlatformMisc::GetUserIndexForPlatformUser(PlatformUserId), *LoginResult.GetErrorValue().GetLogString());
+			UE_LOG(LogTemp, Warning, TEXT("[FLobbiesCommon::FunctionalTestLoginUser] Login failed. PlatformUserId: %d, Result: %s"), FPlatformMisc::GetUserIndexForPlatformUser(PlatformUserId), *LoginResult.GetErrorValue().GetLogString());
 			Promise->EmplaceValue(Errors::RequestFailure());
 		}
 		else
 		{
-			UE_LOG(LogLobbies, Log, TEXT("[FLobbiesCommon::FunctionalTestLoginUser] Login Succeeded. PlatformUserId: %d, User: %s"), FPlatformMisc::GetUserIndexForPlatformUser(PlatformUserId), *ToLogString(LoginResult.GetOkValue().AccountInfo->UserId));
+			UE_LOG(LogTemp, Log, TEXT("[FLobbiesCommon::FunctionalTestLoginUser] Login Succeeded. PlatformUserId: %d, User: %s"), FPlatformMisc::GetUserIndexForPlatformUser(PlatformUserId), *ToLogString(LoginResult.GetOkValue().AccountInfo->UserId));
 			Promise->EmplaceValue(FFunctionalTestLoginUser::Result{LoginResult.GetOkValue().AccountInfo});
 		}
 	});
@@ -922,12 +920,12 @@ TFuture<TOnlineResult<FLobbiesCommon::FFunctionalTestLoginUser>> FLobbiesCommon:
 
 TFuture<TOnlineResult<FLobbiesCommon::FFunctionalTestLogoutUser>> FLobbiesCommon::FunctionalTestLogoutUser(FFunctionalTestLogoutUser::Params&& Params)
 {
-	UE_LOG(LogLobbies, Log, TEXT("[FLobbiesCommon::FunctionalTestLogoutUser] Logging out user. PlatformUserId: %d"), FPlatformMisc::GetUserIndexForPlatformUser(Params.PlatformUserId));
+	UE_LOG(LogTemp, Log, TEXT("[FLobbiesCommon::FunctionalTestLogoutUser] Logging out user. PlatformUserId: %d"), FPlatformMisc::GetUserIndexForPlatformUser(Params.PlatformUserId));
 
 	IAuth* AuthInterface = Services.Get<IAuth>();
 	if (!AuthInterface)
 	{
-		UE_LOG(LogLobbies, Warning, TEXT("[FLobbiesCommon::FunctionalTestLogoutUser] Logout failed. PlatformUserId: %d"), FPlatformMisc::GetUserIndexForPlatformUser(Params.PlatformUserId));
+		UE_LOG(LogTemp, Warning, TEXT("[FLobbiesCommon::FunctionalTestLogoutUser] Logout failed. PlatformUserId: %d"), FPlatformMisc::GetUserIndexForPlatformUser(Params.PlatformUserId));
 		return MakeFulfilledPromise<TOnlineResult<FFunctionalTestLogoutUser>>(Errors::MissingInterface()).GetFuture();
 	}
 
@@ -939,12 +937,12 @@ TFuture<TOnlineResult<FLobbiesCommon::FFunctionalTestLogoutUser>> FLobbiesCommon
 		// Ignore errors for now. This function should ignore logged out users.
 		if (Result.GetErrorValue() == Errors::Unknown())
 		{
-			UE_LOG(LogLobbies, Log, TEXT("[FLobbiesCommon::FunctionalTestLogoutUser] Logout Succeeded - already logged out. PlatformUserId: %d"), FPlatformMisc::GetUserIndexForPlatformUser(Params.PlatformUserId));
+			UE_LOG(LogTemp, Log, TEXT("[FLobbiesCommon::FunctionalTestLogoutUser] Logout Succeeded - already logged out. PlatformUserId: %d"), FPlatformMisc::GetUserIndexForPlatformUser(Params.PlatformUserId));
 			return MakeFulfilledPromise<TOnlineResult<FFunctionalTestLogoutUser>>(FFunctionalTestLogoutUser::Result{}).GetFuture();
 		}
 		else
 		{
-			UE_LOG(LogLobbies, Warning, TEXT("[FLobbiesCommon::FunctionalTestLogoutUser] Logout failed. PlatformUserId: %d, Result: %s"), FPlatformMisc::GetUserIndexForPlatformUser(Params.PlatformUserId), *Result.GetErrorValue().GetLogString());
+			UE_LOG(LogTemp, Warning, TEXT("[FLobbiesCommon::FunctionalTestLogoutUser] Logout failed. PlatformUserId: %d, Result: %s"), FPlatformMisc::GetUserIndexForPlatformUser(Params.PlatformUserId), *Result.GetErrorValue().GetLogString());
 			return MakeFulfilledPromise<TOnlineResult<FFunctionalTestLogoutUser>>(Errors::Unknown(MoveTemp(Result.GetErrorValue()))).GetFuture();
 		}
 	}
@@ -966,12 +964,12 @@ TFuture<TOnlineResult<FLobbiesCommon::FFunctionalTestLogoutUser>> FLobbiesCommon
 	{
 		if (LogoutResult.IsError())
 		{
-			UE_LOG(LogLobbies, Warning, TEXT("[FLobbiesCommon::FunctionalTestLogoutUser] Logout failed. PlatformUserId: %d, Result: %s"), FPlatformMisc::GetUserIndexForPlatformUser(PlatformUserId), *LogoutResult.GetErrorValue().GetLogString());
+			UE_LOG(LogTemp, Warning, TEXT("[FLobbiesCommon::FunctionalTestLogoutUser] Logout failed. PlatformUserId: %d, Result: %s"), FPlatformMisc::GetUserIndexForPlatformUser(PlatformUserId), *LogoutResult.GetErrorValue().GetLogString());
 			Promise->EmplaceValue(Errors::RequestFailure());
 		}
 		else
 		{
-			UE_LOG(LogLobbies, Log, TEXT("[FLobbiesCommon::FunctionalTestLogoutUser] Logout Succeeded. PlatformUserId: %d"), FPlatformMisc::GetUserIndexForPlatformUser(PlatformUserId));
+			UE_LOG(LogTemp, Log, TEXT("[FLobbiesCommon::FunctionalTestLogoutUser] Logout Succeeded. PlatformUserId: %d"), FPlatformMisc::GetUserIndexForPlatformUser(PlatformUserId));
 			Promise->EmplaceValue(FFunctionalTestLogoutUser::Result{});
 		}
 	});
@@ -981,7 +979,7 @@ TFuture<TOnlineResult<FLobbiesCommon::FFunctionalTestLogoutUser>> FLobbiesCommon
 
 TFuture<TOnlineResult<FLobbiesCommon::FFunctionalTestLogoutAllUsers>> FLobbiesCommon::FunctionalTestLogoutAllUsers(FFunctionalTestLogoutAllUsers::Params&& Params)
 {
-	UE_LOG(LogLobbies, Log, TEXT("[FLobbiesCommon::FFunctionalTestLogoutAllUsers] Logging out all users."));
+	UE_LOG(LogTemp, Log, TEXT("[FLobbiesCommon::FFunctionalTestLogoutAllUsers] Logging out all users."));
 
 	TArray<TFuture<TOnlineResult<FLobbiesCommon::FFunctionalTestLogoutUser>>> LogoutFutures;
 	for (int32 index = 0; index < MAX_LOCAL_PLAYERS; ++index)
@@ -1003,12 +1001,12 @@ TFuture<TOnlineResult<FLobbiesCommon::FFunctionalTestLogoutAllUsers>> FLobbiesCo
 
 		if (HasAnyError)
 		{
-			UE_LOG(LogLobbies, Warning, TEXT("[FLobbiesCommon::FunctionalTestLogoutUser] Logout all users Failed."));
+			UE_LOG(LogTemp, Warning, TEXT("[FLobbiesCommon::FunctionalTestLogoutUser] Logout all users Failed."));
 			Promise->EmplaceValue(Errors::RequestFailure());
 		}
 		else
 		{
-			UE_LOG(LogLobbies, Log, TEXT("[FLobbiesCommon::FFunctionalTestLogoutAllUsers] Logout all users Succeeded."));
+			UE_LOG(LogTemp, Log, TEXT("[FLobbiesCommon::FFunctionalTestLogoutAllUsers] Logout all users Succeeded."));
 			Promise->EmplaceValue(FFunctionalTestLogoutAllUsers::Result{});
 		}
 	});
diff --git a/Engine/Plugins/Online/Experimental/OnlineServices/Source/OnlineServicesCommon/Public/Online/LobbiesCommon.h b/Engine/Plugins/Online/Experimental/OnlineServices/Source/OnlineServicesCommon/Public/Online/LobbiesCommon.h
--- a/Engine/Plugins/Online/Experimental/OnlineServices/Source/OnlineServicesCommon/Public/Online/LobbiesCommon.h
+++ b/Engine/Plugins/Online/Experimental/OnlineServices/Source/OnlineServicesCommon/Public/Online/LobbiesCommon.h
@@ -7,8 +7,6 @@
 
 #define LOBBIES_FUNCTIONAL_TEST_ENABLED !UE_BUILD_SHIPPING
 
-DECLARE_LOG_CATEGORY_EXTERN(LogLobbies, Log, All);
-
 namespace UE::Online {
 
 class FOnlineServicesCommon;
@@ -279,19 +277,19 @@ protected:
 			TPromise<void> Promise;
 			auto Future = Promise.GetFuture();
 
-			UE_LOG(LogLobbies, Log, TEXT("[FLobbiesCommon::ConsumeStepResult] Starting secondary operation %s."), SecondaryOpType::Name);
+			UE_LOG(LogTemp, Log, TEXT("[FLobbiesCommon::ConsumeStepResult] Starting secondary operation %s."), SecondaryOpType::Name);
 
 			(this->*Func)(MoveTempIfPossible(InParams))
 			.Then([Promise = MoveTemp(Promise), Op = InAsyncOp.AsShared()](TFuture<TOnlineResult<SecondaryOpType>>&& Future) mutable
 			{
 				if (Future.Get().IsError())
 				{
-					UE_LOG(LogLobbies, Log, TEXT("[FLobbiesCommon::ConsumeStepResult] Failed secondary operation %s. Error: %s"), SecondaryOpType::Name, *Future.Get().GetErrorValue().GetLogString());
+					UE_LOG(LogTemp, Log, TEXT("[FLobbiesCommon::ConsumeStepResult] Failed secondary operation %s. Error: %s"), SecondaryOpType::Name, *Future.Get().GetErrorValue().GetLogString());
 					Op->SetError(Errors::RequestFailure(MoveTempIfPossible(Future.Get().GetErrorValue())));
 				}
 				else
 				{
-					UE_LOG(LogLobbies, Log, TEXT("[FLobbiesCommon::ConsumeStepResult] Completed secondary operation %s."), SecondaryOpType::Name);
+					UE_LOG(LogTemp, Log, TEXT("[FLobbiesCommon::ConsumeStepResult] Completed secondary operation %s."), SecondaryOpType::Name);
 				}
 				Promise.EmplaceValue();
 			});
@@ -308,19 +306,19 @@ protected:
 			TPromise<void> Promise;
 			auto Future = Promise.GetFuture();
 
-			UE_LOG(LogLobbies, Log, TEXT("[FLobbiesCommon::ConsumeStepResult] Starting secondary operation %s."), SecondaryOpType::Name);
+			UE_LOG(LogTemp, Log, TEXT("[FLobbiesCommon::ConsumeStepResult] Starting secondary operation %s."), SecondaryOpType::Name);
 
 			(this->*Func)(MoveTempIfPossible(InParams))
 			.Then([Promise = MoveTemp(Promise), Op = InAsyncOp.AsShared()](TFuture<TOnlineResult<SecondaryOpType>>&& Future) mutable
 			{
 				if (Future.Get().IsError())
 				{
-					UE_LOG(LogLobbies, Log, TEXT("[FLobbiesCommon::ConsumeStepResult] Failed secondary operation %s. Error: %s"), SecondaryOpType::Name, *Future.Get().GetErrorValue().GetLogString());
+					UE_LOG(LogTemp, Log, TEXT("[FLobbiesCommon::ConsumeStepResult] Failed secondary operation %s. Error: %s"), SecondaryOpType::Name, *Future.Get().GetErrorValue().GetLogString());
 					Op->SetError(Errors::RequestFailure(MoveTempIfPossible(Future.Get().GetErrorValue())));
 				}
 				else
 				{
-					UE_LOG(LogLobbies, Log, TEXT("[FLobbiesCommon::ConsumeStepResult] Completed secondary operation %s."), SecondaryOpType::Name);
+					UE_LOG(LogTemp, Log, TEXT("[FLobbiesCommon::ConsumeStepResult] Completed secondary operation %s."), SecondaryOpType::Name);
 				}
 				Promise.EmplaceValue();
 			});
@@ -337,19 +335,19 @@ protected:
 			TPromise<void> Promise;
 			auto Future = Promise.GetFuture();
 
-			UE_LOG(LogLobbies, Log, TEXT("[FLobbiesCommon::CaptureStepResult] Starting secondary operation %s."), SecondaryOpType::Name);
+			UE_LOG(LogTemp, Log, TEXT("[FLobbiesCommon::CaptureStepResult] Starting secondary operation %s."), SecondaryOpType::Name);
 
 			(this->*Func)(MoveTempIfPossible(InParams))
 			.Then([Promise = MoveTemp(Promise), Op = InAsyncOp.AsShared(), ResultKey](TFuture<TOnlineResult<SecondaryOpType>>&& Future) mutable
 			{
 				if (Future.Get().IsError())
 				{
-					UE_LOG(LogLobbies, Log, TEXT("[FLobbiesCommon::CaptureStepResult] Failed secondary operation %s. Error: %s"), SecondaryOpType::Name, *Future.Get().GetErrorValue().GetLogString());
+					UE_LOG(LogTemp, Log, TEXT("[FLobbiesCommon::CaptureStepResult] Failed secondary operation %s. Error: %s"), SecondaryOpType::Name, *Future.Get().GetErrorValue().GetLogString());
 					Op->SetError(Errors::RequestFailure(MoveTempIfPossible(Future.Get().GetErrorValue())));
 				}
 				else
 				{
-					UE_LOG(LogLobbies, Log, TEXT("[FLobbiesCommon::CaptureStepResult] Captured secondary operation %s as %s"), SecondaryOpType::Name, *ResultKey);
+					UE_LOG(LogTemp, Log, TEXT("[FLobbiesCommon::CaptureStepResult] Captured secondary operation %s as %s"), SecondaryOpType::Name, *ResultKey);
 					Op->Data.Set(ResultKey, MoveTempIfPossible(Future.Get().GetOkValue()));
 				}
 				Promise.EmplaceValue();
@@ -367,19 +365,19 @@ protected:
 			TPromise<void> Promise;
 			auto Future = Promise.GetFuture();
 
-			UE_LOG(LogLobbies, Log, TEXT("[FLobbiesCommon::CaptureStepResult] Starting secondary operation %s."), SecondaryOpType::Name);
+			UE_LOG(LogTemp, Log, TEXT("[FLobbiesCommon::CaptureStepResult] Starting secondary operation %s."), SecondaryOpType::Name);
 
 			(this->*Func)(MoveTempIfPossible(InParams))
 			.Then([Promise = MoveTemp(Promise), Op = InAsyncOp.AsShared(), ResultKey](TFuture<TOnlineResult<SecondaryOpType>>&& Future) mutable
 			{
 				if (Future.Get().IsError())
 				{
-					UE_LOG(LogLobbies, Log, TEXT("[FLobbiesCommon::CaptureStepResult] Failed secondary operation %s. Error: %s"), SecondaryOpType::Name, *Future.Get().GetErrorValue().GetLogString());
+					UE_LOG(LogTemp, Log, TEXT("[FLobbiesCommon::CaptureStepResult] Failed secondary operation %s. Error: %s"), SecondaryOpType::Name, *Future.Get().GetErrorValue().GetLogString());
 					Op->SetError(Errors::RequestFailure(MoveTempIfPossible(Future.Get().GetErrorValue())));
 				}
 				else
 				{
-					UE_LOG(LogLobbies, Log, TEXT("[FLobbiesCommon::CaptureStepResult] Captured secondary operation %s as %s"), SecondaryOpType::Name, *ResultKey);
+					UE_LOG(LogTemp, Log, TEXT("[FLobbiesCommon::CaptureStepResult] Captured secondary operation %s as %s"), SecondaryOpType::Name, *ResultKey);
 					Op->Data.Set(ResultKey, MoveTempIfPossible(Future.Get().GetOkValue()));
 				}
 				Promise.EmplaceValue();
@@ -397,19 +395,19 @@ protected:
 			auto Promise = MakeShared<TPromise<void>>();
 			auto Future = Promise->GetFuture();
 
-			UE_LOG(LogLobbies, Log, TEXT("[FLobbiesCommon::ConsumeOperationStepResult] Starting secondary operation %s."), SecondaryOpType::Name);
+			UE_LOG(LogTemp, Log, TEXT("[FLobbiesCommon::ConsumeOperationStepResult] Starting secondary operation %s."), SecondaryOpType::Name);
 
 			(this->*Func)(MoveTempIfPossible(InParams))
 			.OnComplete([Promise, Op = InAsyncOp.AsShared()](const TOnlineResult<SecondaryOpType>& Result) mutable -> void
 			{
 				if (Result.IsError())
 				{
-					UE_LOG(LogLobbies, Log, TEXT("[FLobbiesCommon::ConsumeOperationStepResult] Failed secondary operation %s. Error: %s"), SecondaryOpType::Name, *Result.GetErrorValue().GetLogString());
+					UE_LOG(LogTemp, Log, TEXT("[FLobbiesCommon::ConsumeOperationStepResult] Failed secondary operation %s. Error: %s"), SecondaryOpType::Name, *Result.GetErrorValue().GetLogString());
 					Op->SetError(Errors::RequestFailure(MoveTempIfPossible(Result.GetErrorValue())));
 				}
 				else
 				{
-					UE_LOG(LogLobbies, Log, TEXT("[FLobbiesCommon::ConsumeOperationStepResult] Completed secondary operation %s."), SecondaryOpType::Name);
+					UE_LOG(LogTemp, Log, TEXT("[FLobbiesCommon::ConsumeOperationStepResult] Completed secondary operation %s."), SecondaryOpType::Name);
 				}
 				Promise->EmplaceValue();
 			});
@@ -426,19 +424,19 @@ protected:
 			auto Promise = MakeShared<TPromise<void>>();
 			auto Future = Promise->GetFuture();
 
-			UE_LOG(LogLobbies, Log, TEXT("[FLobbiesCommon::ConsumeOperationStepResult] Starting secondary operation %s."), SecondaryOpType::Name);
+			UE_LOG(LogTemp, Log, TEXT("[FLobbiesCommon::ConsumeOperationStepResult] Starting secondary operation %s."), SecondaryOpType::Name);
 
 			(this->*Func)(MoveTempIfPossible(InParams))
 			.OnComplete([Promise, Op = InAsyncOp.AsShared()](const TOnlineResult<SecondaryOpType>& Result) mutable
 			{
 				if (Result.IsError())
 				{
-					UE_LOG(LogLobbies, Log, TEXT("[FLobbiesCommon::ConsumeOperationStepResult] Failed secondary operation %s. Error: %s"), SecondaryOpType::Name, *Result.GetErrorValue().GetLogString());
+					UE_LOG(LogTemp, Log, TEXT("[FLobbiesCommon::ConsumeOperationStepResult] Failed secondary operation %s. Error: %s"), SecondaryOpType::Name, *Result.GetErrorValue().GetLogString());
 					Op->SetError(Errors::RequestFailure(MoveTempIfPossible(Result.GetErrorValue())));
 				}
 				else
 				{
-					UE_LOG(LogLobbies, Log, TEXT("[FLobbiesCommon::ConsumeOperationStepResult] Completed secondary operation %s."), SecondaryOpType::Name);
+					UE_LOG(LogTemp, Log, TEXT("[FLobbiesCommon::ConsumeOperationStepResult] Completed secondary operation %s."), SecondaryOpType::Name);
 				}
 				Promise->EmplaceValue();
 			});
@@ -455,19 +453,19 @@ protected:
 			auto Promise = MakeShared<TPromise<void>>();
 			auto Future = Promise->GetFuture();
 
-			UE_LOG(LogLobbies, Log, TEXT("[FLobbiesCommon::CaptureOperationStepResult] Starting secondary operation %s."), SecondaryOpType::Name);
+			UE_LOG(LogTemp, Log, TEXT("[FLobbiesCommon::CaptureOperationStepResult] Starting secondary operation %s."), SecondaryOpType::Name);
 
 			(this->*Func)(MoveTempIfPossible(InParams))
 			.OnComplete([Promise, Op = InAsyncOp.AsShared(), ResultKey](const TOnlineResult<SecondaryOpType>& Result) mutable -> void
 			{
 				if (Result.IsError())
 				{
-					UE_LOG(LogLobbies, Log, TEXT("[FLobbiesCommon::CaptureOperationStepResult] Failed secondary operation %s. Error: %s"), SecondaryOpType::Name, *Result.GetErrorValue().GetLogString());
+					UE_LOG(LogTemp, Log, TEXT("[FLobbiesCommon::CaptureOperationStepResult] Failed secondary operation %s. Error: %s"), SecondaryOpType::Name, *Result.GetErrorValue().GetLogString());
 					Op->SetError(Errors::RequestFailure(MoveTempIfPossible(Result.GetErrorValue())));
 				}
 				else
 				{
-					UE_LOG(LogLobbies, Log, TEXT("[FLobbiesCommon::CaptureOperationStepResult] Captured secondary operation %s as %s"), SecondaryOpType::Name, *ResultKey);
+					UE_LOG(LogTemp, Log, TEXT("[FLobbiesCommon::CaptureOperationStepResult] Captured secondary operation %s as %s"), SecondaryOpType::Name, *ResultKey);
 					Op->Data.Set(ResultKey, Result.GetOkValue());
 				}
 				Promise->EmplaceValue();
@@ -485,19 +483,19 @@ protected:
 			auto Promise = MakeShared<TPromise<void>>();
 			auto Future = Promise->GetFuture();
 
-			UE_LOG(LogLobbies, Log, TEXT("[FLobbiesCommon::CaptureOperationStepResult] Starting secondary operation %s."), SecondaryOpType::Name);
+			UE_LOG(LogTemp, Log, TEXT("[FLobbiesCommon::CaptureOperationStepResult] Starting secondary operation %s."), SecondaryOpType::Name);
 
 			(this->*Func)(MoveTempIfPossible(InParams))
 			.OnComplete([Promise, Op = InAsyncOp.AsShared(), ResultKey](const TOnlineResult<SecondaryOpType>& Result) mutable
 			{
 				if (Result.IsError())
 				{
-					UE_LOG(LogLobbies, Log, TEXT("[FLobbiesCommon::CaptureOperationStepResult] Failed secondary operation %s. Error: %s"), SecondaryOpType::Name, *Result.GetErrorValue().GetLogString());
+					UE_LOG(LogTemp, Log, TEXT("[FLobbiesCommon::CaptureOperationStepResult] Failed secondary operation %s. Error: %s"), SecondaryOpType::Name, *Result.GetErrorValue().GetLogString());
 					Op->SetError(Errors::RequestFailure(MoveTempIfPossible(Result.GetErrorValue())));
 				}
 				else
 				{
-					UE_LOG(LogLobbies, Log, TEXT("[FLobbiesCommon::CaptureOperationStepResult] Captured secondary operation %s as %s"), SecondaryOpType::Name, *ResultKey);
+					UE_LOG(LogTemp, Log, TEXT("[FLobbiesCommon::CaptureOperationStepResult] Captured secondary operation %s as %s"), SecondaryOpType::Name, *ResultKey);
 					Op->Data.Set(ResultKey, Result.GetOkValue());
 				}
 				Promise->EmplaceValue();
diff --git a/Engine/Plugins/Online/Experimental/OnlineServicesEOS/Source/Private/LobbiesEOS.cpp b/Engine/Plugins/Online/Experimental/OnlineServicesEOS/Source/Private/LobbiesEOS.cpp
--- a/Engine/Plugins/Online/Experimental/OnlineServicesEOS/Source/Private/LobbiesEOS.cpp
+++ b/Engine/Plugins/Online/Experimental/OnlineServicesEOS/Source/Private/LobbiesEOS.cpp
@@ -11,8 +11,6 @@
 
 #include "eos_lobby.h"
 
-DEFINE_LOG_CATEGORY(LogLobbies);
-
 namespace UE::Online {
 
 static const FString LobbyDataKeyName = TEXT("LobbyData");
@@ -123,13 +121,13 @@ TOnlineAsyncOpHandle<FCreateLobby> FLobbiesEOS::CreateLobby(FCreateLobby::Params
 		CreateLobbyOptions.bDisableHostMigration = false; // todo: handle
 		CreateLobbyOptions.bEnableRTCRoom = false; // todo: handle
 
-		UE_LOG(LogLobbies, Warning, TEXT("[FLobbiesEOS::EOS_Lobby_CreateLobby] Start. Member: %s"), *ToLogString(Params.LocalUserId));
+		UE_LOG(LogTemp, Warning, TEXT("[FLobbiesEOS::EOS_Lobby_CreateLobby] Start. Member: %s"), *ToLogString(Params.LocalUserId));
 
 		return EOS_Async<EOS_Lobby_CreateLobbyCallbackInfo>(EOS_Lobby_CreateLobby, LobbyPrerequisites->LobbyInterfaceHandle, CreateLobbyOptions); 
 	})
 	.Then([this](TOnlineAsyncOp<FCreateLobby>& InAsyncOp, const EOS_Lobby_CreateLobbyCallbackInfo* Data)
 	{
-		UE_LOG(LogLobbies, Warning, TEXT("[FLobbiesEOS::EOS_Lobby_CreateLobby] Complete. Result: %d"), Data->ResultCode);
+		UE_LOG(LogTemp, Warning, TEXT("[FLobbiesEOS::EOS_Lobby_CreateLobby] Complete. Result: %d"), Data->ResultCode);
 
 		if (Data->ResultCode != EOS_EResult::EOS_Success)
 		{
@@ -688,7 +686,7 @@ void FLobbiesEOS::HandleLobbyUpdated(const EOS_Lobby_LobbyUpdateReceivedCallback
 
 	if (TSharedPtr<FLobbyDataEOS> LobbyData = LobbyDataRegistry->Find(Data->LobbyId))
 	{
-		UE_LOG(LogLobbies, Warning, TEXT("[FLobbiesEOS::HandleLobbyUpdated] Received update. Lobby: %s"), *LobbyData->GetLobbyId());
+		UE_LOG(LogTemp, Warning, TEXT("[FLobbiesEOS::HandleLobbyUpdated] Received update. Lobby: %s"), *LobbyData->GetLobbyId());
 	}
 }
 
@@ -710,7 +708,7 @@ void FLobbiesEOS::HandleLobbyMemberUpdated(const EOS_Lobby_LobbyMemberUpdateRece
 	{
 		if (TSharedPtr<FLobbyDataEOS> LobbyData = LobbyDataRegistry->Find(Data->LobbyId))
 		{
-			UE_LOG(LogLobbies, Warning, TEXT("[FLobbiesEOS::HandleLobbyMemberUpdated] Received update. Lobby: %s, Member: %s"),
+			UE_LOG(LogTemp, Warning, TEXT("[FLobbiesEOS::HandleLobbyMemberUpdated] Received update. Lobby: %s, Member: %s"),
 				*LobbyData->GetLobbyId(), *ToLogString(LocalUserId));
 		}
 	}
@@ -736,7 +734,7 @@ void FLobbiesEOS::HandleLobbyMemberStatusReceived(const EOS_Lobby_LobbyMemberSta
 	{
 		if (TSharedPtr<FLobbyDataEOS> LobbyData = LobbyDataRegistry->Find(Data->LobbyId))
 		{
-			UE_LOG(LogLobbies, Warning, TEXT("[FLobbiesEOS::HandleLobbyMemberStatusReceived] Received update. Lobby: %s, Member: %s, Statsu: %d"),
+			UE_LOG(LogTemp, Warning, TEXT("[FLobbiesEOS::HandleLobbyMemberStatusReceived] Received update. Lobby: %s, Member: %s, Statsu: %d"),
 				*LobbyData->GetLobbyId(), *ToLogString(LocalUserId), Data->CurrentStatus);
 		}
 	}
@@ -767,12 +765,12 @@ void FLobbiesEOS::HandleLobbyInviteReceived(const EOS_Lobby_LobbyInviteReceivedC
 			if (Future.Get().IsError())
 			{
 				// Todo: Log / queue a manual fetch of invitations.
-				UE_LOG(LogLobbies, Warning, TEXT("[FLobbiesEOS::HandleLobbyInviteReceived] Failed to receive invite. Error: %s"),
+				UE_LOG(LogTemp, Warning, TEXT("[FLobbiesEOS::HandleLobbyInviteReceived] Failed to receive invite. Error: %s"),
 					*Future.Get().GetErrorValue().GetLogString());
 			}
 			else
 			{
-				UE_LOG(LogLobbies, Warning, TEXT("[FLobbiesEOS::HandleLobbyInviteReceived] Received invite. Id: %s, Lobby: %s, Receiver: %s, Sender: %s"),
+				UE_LOG(LogTemp, Warning, TEXT("[FLobbiesEOS::HandleLobbyInviteReceived] Received invite. Id: %s, Lobby: %s, Receiver: %s, Sender: %s"),
 					*Future.Get().GetOkValue()->GetInviteId(),
 					*Future.Get().GetOkValue()->GetLobbyData()->GetLobbyId(),
 					*ToLogString(Future.Get().GetOkValue()->GetReceiver()),
@@ -1106,14 +1104,14 @@ TOnlineAsyncOpHandle<FLobbiesEOS::FJoinLobbyMemberImpl> FLobbiesEOS::JoinLobbyMe
 		JoinLobbyOptions.bPresenceEnabled = false;
 		JoinLobbyOptions.LocalRTCOptions = nullptr;
 
-		UE_LOG(LogLobbies, Warning, TEXT("[FLobbiesEOS::EOS_Lobby_JoinLobby] Start. Lobby: %s, Member: %s"),
+		UE_LOG(LogTemp, Warning, TEXT("[FLobbiesEOS::EOS_Lobby_JoinLobby] Start. Lobby: %s, Member: %s"),
 			*LobbyData->GetLobbyId(), *ToLogString(Params.LocalUserId));
 
 		return EOS_Async<EOS_Lobby_JoinLobbyCallbackInfo>(EOS_Lobby_JoinLobby, LobbyPrerequisites->LobbyInterfaceHandle, JoinLobbyOptions);
 	})
 	.Then([this](TOnlineAsyncOp<FJoinLobbyMemberImpl>& InAsyncOp, const EOS_Lobby_JoinLobbyCallbackInfo* Data)
 	{
-		UE_LOG(LogLobbies, Warning, TEXT("[FLobbiesEOS::EOS_Lobby_JoinLobby] Complete. Result: %d"), Data->ResultCode);
+		UE_LOG(LogTemp, Warning, TEXT("[FLobbiesEOS::EOS_Lobby_JoinLobby] Complete. Result: %d"), Data->ResultCode);
 		if (Data->ResultCode != EOS_EResult::EOS_Success)
 		{
 			// TODO: Error codes
@@ -1437,7 +1435,7 @@ TFuture<TDefaultErrorResult<FLobbiesEOS::FModifyLobbyDataImpl>> FLobbiesEOS::Mod
 
 	// todo: check local user is lobby owner
 
-	UE_LOG(LogLobbies, Warning, TEXT("[FLobbiesEOS::ModifyLobbyDataImpl] Start. Lobby: %s, Member: %s"),
+	UE_LOG(LogTemp, Warning, TEXT("[FLobbiesEOS::ModifyLobbyDataImpl] Start. Lobby: %s, Member: %s"),
 		*Params.LobbyData->GetLobbyId(), *ToLogString(Params.LocalUserId));
 
 	TPromise<TDefaultErrorResult<FModifyLobbyDataImpl>> Promise;
@@ -1446,7 +1444,7 @@ TFuture<TDefaultErrorResult<FLobbiesEOS::FModifyLobbyDataImpl>> FLobbiesEOS::Mod
 	DetailsResult.GetOkValue()->ApplyLobbyDataUpdates(Params.LocalUserId, MoveTemp(*Params.Changes))
 	.Then([Promise = MoveTemp(Promise)](TFuture<EOS_EResult>&& Future) mutable
 	{
-		UE_LOG(LogLobbies, Warning, TEXT("[FLobbiesEOS::ModifyLobbyDataImpl] Complete. Result: %d"), Future.Get());
+		UE_LOG(LogTemp, Warning, TEXT("[FLobbiesEOS::ModifyLobbyDataImpl] Complete. Result: %d"), Future.Get());
 
 		// Todo: Handle "no change" better.
 		if (Future.Get() != EOS_EResult::EOS_Success && Future.Get() != EOS_EResult::EOS_NoChange)
@@ -1483,7 +1481,7 @@ TFuture<TDefaultErrorResult<FLobbiesEOS::FModifyLobbyMemberDataImpl>> FLobbiesEO
 		return MakeFulfilledPromise<TDefaultErrorResult<FModifyLobbyMemberDataImpl>>(Errors::Unknown(MoveTemp(DetailsResult.GetErrorValue()))).GetFuture();
 	}
 
-	UE_LOG(LogLobbies, Warning, TEXT("[FLobbiesEOS::ModifyLobbyMemberDataImpl] Start. Lobby: %s, Member: %s"),
+	UE_LOG(LogTemp, Warning, TEXT("[FLobbiesEOS::ModifyLobbyMemberDataImpl] Start. Lobby: %s, Member: %s"),
 		*Params.LobbyData->GetLobbyId(), *ToLogString(Params.LocalUserId));
 
 	TPromise<TDefaultErrorResult<FModifyLobbyMemberDataImpl>> Promise;
@@ -1492,7 +1490,7 @@ TFuture<TDefaultErrorResult<FLobbiesEOS::FModifyLobbyMemberDataImpl>> FLobbiesEO
 	DetailsResult.GetOkValue()->ApplyLobbyMemberDataUpdates(Params.LocalUserId, MoveTemp(*Params.Changes))
 	.Then([Promise = MoveTemp(Promise)](TFuture<EOS_EResult>&& Future) mutable
 	{
-		UE_LOG(LogLobbies, Warning, TEXT("[FLobbiesEOS::ModifyLobbyMemberDataImpl] Complete. Result: %d"), Future.Get());
+		UE_LOG(LogTemp, Warning, TEXT("[FLobbiesEOS::ModifyLobbyMemberDataImpl] Complete. Result: %d"), Future.Get());
 
 		// Todo: Handle "no change" better.
 		if (Future.Get() != EOS_EResult::EOS_Success && Future.Get() != EOS_EResult::EOS_NoChange)
