Fix issue where queued ops would not run on previous op completion.

diff --git a/Engine/Plugins/Online/Experimental/OnlineServices/Source/OnlineServicesCommon/Public/Online/OnlineAsyncOpQueue.h b/Engine/Plugins/Online/Experimental/OnlineServices/Source/OnlineServicesCommon/Public/Online/OnlineAsyncOpQueue.h
--- a/Engine/Plugins/Online/Experimental/OnlineServices/Source/OnlineServicesCommon/Public/Online/OnlineAsyncOpQueue.h
+++ b/Engine/Plugins/Online/Experimental/OnlineServices/Source/OnlineServicesCommon/Public/Online/OnlineAsyncOpQueue.h
@@ -138,7 +138,7 @@ public:
 				// push to the parent parallel queue
 				ParentQueue.Enqueue(Operation.GetValue());
 				// when the operation completes, start the next one
-				InFlightOperation->OnComplete().Add([this]()
+				ContinuationDelegateHandle = InFlightOperation->OnComplete().Add([this]()
 					{
 						InFlightOperation.Reset();
 						TryStartOperations();
@@ -151,6 +151,7 @@ protected:
 	FOnlineAsyncOpQueueParallel& ParentQueue;
 
 	TSharedPtr<IWrappedOp> InFlightOperation;
+	FOnlineEventDelegateHandle ContinuationDelegateHandle;
 };
 
 /* UE::Online */ }
