Add Lobbies accessor to IOnlineServices
Move cpp to private

diff --git a/Engine/Plugins/Online/Experimental/OnlineServices/Source/OnlineServicesCommon/Private/Online/OnlineServicesCommon.cpp b/Engine/Plugins/Online/Experimental/OnlineServices/Source/OnlineServicesCommon/Private/Online/OnlineServicesCommon.cpp
--- a/Engine/Plugins/Online/Experimental/OnlineServices/Source/OnlineServicesCommon/Private/Online/OnlineServicesCommon.cpp
+++ b/Engine/Plugins/Online/Experimental/OnlineServices/Source/OnlineServicesCommon/Private/Online/OnlineServicesCommon.cpp
@@ -59,6 +59,11 @@ IExternalUIPtr FOnlineServicesCommon::GetExternalUIInterface()
 	return IExternalUIPtr(AsShared(), Get<IExternalUI>());
 }
 
+ILobbiesPtr FOnlineServicesCommon::GetLobbiesInterface()
+{
+	return ILobbiesPtr(AsShared(), Get<ILobbies>());
+}
+
 void FOnlineServicesCommon::RegisterComponents()
 {
 }
diff --git a/Engine/Plugins/Online/Experimental/OnlineServices/Source/OnlineServicesCommon/Public/Online/OnlineServicesCommon.h b/Engine/Plugins/Online/Experimental/OnlineServices/Source/OnlineServicesCommon/Public/Online/OnlineServicesCommon.h
--- a/Engine/Plugins/Online/Experimental/OnlineServices/Source/OnlineServicesCommon/Public/Online/OnlineServicesCommon.h
+++ b/Engine/Plugins/Online/Experimental/OnlineServices/Source/OnlineServicesCommon/Public/Online/OnlineServicesCommon.h
@@ -39,6 +39,7 @@ public:
 	virtual IFriendsPtr GetFriendsInterface() override;
 	virtual IPresencePtr GetPresenceInterface() override;
 	virtual IExternalUIPtr GetExternalUIInterface() override;
+	virtual ILobbiesPtr GetLobbiesInterface() override;
 
 	// FOnlineServicesCommon
 
diff --git a/Engine/Plugins/Online/Experimental/OnlineServices/Source/OnlineServicesInterface/Private/Lobbies.cpp b/Engine/Plugins/Online/Experimental/OnlineServices/Source/OnlineServicesInterface/Private/Lobbies.cpp
new file mode 100644
--- /dev/null
+++ b/Engine/Plugins/Online/Experimental/OnlineServices/Source/OnlineServicesInterface/Private/Lobbies.cpp
@@ -0,0 +1,102 @@
+// Copyright Epic Games, Inc. All Rights Reserved.
+
+#include "Online/Lobbies.h"
+
+namespace UE::Online {
+
+const TCHAR* LexToString(ELobbyJoinPolicy Policy)
+{
+	switch (Policy)
+	{
+	case ELobbyJoinPolicy::PublicAdvertised:	return TEXT("PublicAdvertised");
+	case ELobbyJoinPolicy::PublicNotAdvertised:	return TEXT("PublicNotAdvertised");
+	default:									checkNoEntry(); // Intentional fallthrough
+	case ELobbyJoinPolicy::InvitationOnly:		return TEXT("InvitationOnly");
+	}
+}
+
+void LexFromString(ELobbyJoinPolicy& OutPolicy, const TCHAR* InStr)
+{
+	if (FCString::Stricmp(InStr, TEXT("PublicAdvertised")) == 0)
+	{
+		OutPolicy = ELobbyJoinPolicy::PublicAdvertised;
+	}
+	else if (FCString::Stricmp(InStr, TEXT("PublicNotAdvertised")) == 0)
+	{
+		OutPolicy = ELobbyJoinPolicy::PublicNotAdvertised;
+	}
+	else if (FCString::Stricmp(InStr, TEXT("InvitationOnly")) == 0)
+	{
+		OutPolicy = ELobbyJoinPolicy::InvitationOnly;
+	}
+	else
+	{
+		checkNoEntry();
+		OutPolicy = ELobbyJoinPolicy::InvitationOnly;
+	}
+}
+
+const TCHAR* LexToString(ELobbyComparisonOp Comparison)
+{
+	switch (Comparison)
+	{
+	case ELobbyComparisonOp::Equals:			return TEXT("Equals");
+	case ELobbyComparisonOp::NotEquals:			return TEXT("NotEquals");
+	case ELobbyComparisonOp::GreaterThan:		return TEXT("GreaterThan");
+	case ELobbyComparisonOp::GreaterThanEquals:	return TEXT("GreaterThanEquals");
+	case ELobbyComparisonOp::LessThan:			return TEXT("LessThan");
+	case ELobbyComparisonOp::LessThanEquals:	return TEXT("LessThanEquals");
+	case ELobbyComparisonOp::Near:				return TEXT("Near");
+	case ELobbyComparisonOp::In:				return TEXT("In");
+	default:									checkNoEntry(); // Intentional fallthrough
+	case ELobbyComparisonOp::NotIn:				return TEXT("NotIn");
+	}
+}
+
+void LexFromString(ELobbyComparisonOp& OutComparison, const TCHAR* InStr)
+{
+	if (FCString::Stricmp(InStr, TEXT("Equals")) == 0)
+	{
+		OutComparison = ELobbyComparisonOp::Equals;
+	}
+	else if (FCString::Stricmp(InStr, TEXT("NotEquals")) == 0)
+	{
+		OutComparison = ELobbyComparisonOp::NotEquals;
+	}
+	else if (FCString::Stricmp(InStr, TEXT("GreaterThan")) == 0)
+	{
+		OutComparison = ELobbyComparisonOp::GreaterThan;
+	}
+	else if (FCString::Stricmp(InStr, TEXT("GreaterThanEquals")) == 0)
+	{
+		OutComparison = ELobbyComparisonOp::GreaterThanEquals;
+	}
+	else if (FCString::Stricmp(InStr, TEXT("LessThan")) == 0)
+	{
+		OutComparison = ELobbyComparisonOp::LessThan;
+	}
+	else if (FCString::Stricmp(InStr, TEXT("LessThanEquals")) == 0)
+	{
+		OutComparison = ELobbyComparisonOp::LessThanEquals;
+	}
+	else if (FCString::Stricmp(InStr, TEXT("Near")) == 0)
+	{
+		OutComparison = ELobbyComparisonOp::Near;
+	}
+	else if (FCString::Stricmp(InStr, TEXT("In")) == 0)
+	{
+		OutComparison = ELobbyComparisonOp::In;
+	}
+	else
+	{
+		checkNoEntry();
+		OutComparison = ELobbyComparisonOp::In;
+	}
+}
+
+void SortLobbies(const TArray<FFindLobbySearchFilter>& Filters, TArray<TSharedRef<const FLobby>>& Lobbies)
+{
+	// todo
+}
+
+/* UE::Online */ }
diff --git a/Engine/Plugins/Online/Experimental/OnlineServices/Source/OnlineServicesInterface/Public/Online/Lobbies.cpp b/Engine/Plugins/Online/Experimental/OnlineServices/Source/OnlineServicesInterface/Public/Online/Lobbies.cpp
deleted file mode 100644
diff --git a/Engine/Plugins/Online/Experimental/OnlineServices/Source/OnlineServicesInterface/Public/Online/OnlineServices.h b/Engine/Plugins/Online/Experimental/OnlineServices/Source/OnlineServicesInterface/Public/Online/OnlineServices.h
--- a/Engine/Plugins/Online/Experimental/OnlineServices/Source/OnlineServicesInterface/Public/Online/OnlineServices.h
+++ b/Engine/Plugins/Online/Experimental/OnlineServices/Source/OnlineServicesInterface/Public/Online/OnlineServices.h
@@ -15,6 +15,7 @@ using IAuthPtr = TSharedPtr<class IAuth>;
 using IFriendsPtr = TSharedPtr<class IFriends>;
 using IPresencePtr = TSharedPtr<class IPresence>;
 using IExternalUIPtr = TSharedPtr<class IExternalUI>;
+using ILobbiesPtr = TSharedPtr<class ILobbies>;
 
 class ONLINESERVICESINTERFACE_API IOnlineServices
 {
@@ -48,6 +49,12 @@ public:
 	 *
 	 */
 	virtual IExternalUIPtr GetExternalUIInterface() = 0;
+
+	/**
+	 * Get the lobbies implementation
+	 * @return lobbies implementation, may be null if not implemented for this service
+	 */
+	virtual ILobbiesPtr GetLobbiesInterface() = 0;
 };
 
 /**
