Implementation of Connectivity and Privileges interfaces.
Addition of Connection Status accessor and change delegate logic.
Addition of User Privilege query complete delegate logic.
Added missing logic to Auth.h and OnlineUtils.h

diff --git a/Engine/Plugins/Online/Experimental/OnlineServices/Source/OnlineServicesCommon/Private/Online/ConnectivityCommon.cpp b/Engine/Plugins/Online/Experimental/OnlineServices/Source/OnlineServicesCommon/Private/Online/ConnectivityCommon.cpp
new file mode 100644
--- /dev/null
+++ b/Engine/Plugins/Online/Experimental/OnlineServices/Source/OnlineServicesCommon/Private/Online/ConnectivityCommon.cpp
@@ -0,0 +1,34 @@
+// Copyright Epic Games, Inc. All Rights Reserved.
+
+#include "Online/ConnectivityCommon.h"
+
+#include "Online/OnlineAsyncOp.h"
+#include "Online/OnlineErrorDefinitions.h"
+#include "Online/OnlineServicesCommon.h"
+
+namespace UE::Online {
+
+FConnectivityCommon::FConnectivityCommon(FOnlineServicesCommon& InServices)
+	: TOnlineComponent(TEXT("Connectivity"), InServices)
+{
+}
+
+void FConnectivityCommon::RegisterCommands()
+{
+	RegisterCommand(&FConnectivityCommon::GetConnectionStatus);
+}
+
+TOnlineResult<FGetConnectionStatus> FConnectivityCommon::GetConnectionStatus(FGetConnectionStatus::Params&& Params)
+{
+	FGetConnectionStatus::Result Result;
+	Result.Status = ConnectionStatus;
+
+	return TOnlineResult<FGetConnectionStatus>(Result);
+}
+
+TOnlineEvent<void(const FConnectionStatusChanged&)> FConnectivityCommon::OnConnectionStatusChanged()
+{
+	return OnConnectionStatusChangedEvent;
+}
+
+/* UE::Online */}
\ No newline at end of file
diff --git a/Engine/Plugins/Online/Experimental/OnlineServices/Source/OnlineServicesCommon/Private/Online/OnlineServicesCommon.cpp b/Engine/Plugins/Online/Experimental/OnlineServices/Source/OnlineServicesCommon/Private/Online/OnlineServicesCommon.cpp
--- a/Engine/Plugins/Online/Experimental/OnlineServices/Source/OnlineServicesCommon/Private/Online/OnlineServicesCommon.cpp
+++ b/Engine/Plugins/Online/Experimental/OnlineServices/Source/OnlineServicesCommon/Private/Online/OnlineServicesCommon.cpp
@@ -64,6 +64,16 @@ ILobbiesPtr FOnlineServicesCommon::GetLobbiesInterface()
 	return ILobbiesPtr(AsShared(), Get<ILobbies>());
 }
 
+IConnectivityPtr FOnlineServicesCommon::GetConnectivityInterface()
+{
+	return IConnectivityPtr(AsShared(), Get<IConnectivity>());
+}
+
+IPrivilegesPtr FOnlineServicesCommon::GetPrivilegesInterface()
+{
+	return IPrivilegesPtr(AsShared(), Get<IPrivileges>());
+}
+
 void FOnlineServicesCommon::RegisterComponents()
 {
 }
diff --git a/Engine/Plugins/Online/Experimental/OnlineServices/Source/OnlineServicesCommon/Private/Online/PrivilegesCommon.cpp b/Engine/Plugins/Online/Experimental/OnlineServices/Source/OnlineServicesCommon/Private/Online/PrivilegesCommon.cpp
new file mode 100644
--- /dev/null
+++ b/Engine/Plugins/Online/Experimental/OnlineServices/Source/OnlineServicesCommon/Private/Online/PrivilegesCommon.cpp
@@ -0,0 +1,28 @@
+// Copyright Epic Games, Inc. All Rights Reserved.
+
+#include "Online/PrivilegesCommon.h"
+
+#include "Online/OnlineAsyncOp.h"
+#include "Online/OnlineErrorDefinitions.h"
+#include "Online/OnlineServicesCommon.h"
+
+namespace UE::Online {
+
+FPrivilegesCommon::FPrivilegesCommon(FOnlineServicesCommon& InServices)
+	: TOnlineComponent(TEXT("Permission"), InServices)
+{
+}
+
+void FPrivilegesCommon::RegisterCommands()
+{
+	RegisterCommand(&FPrivilegesCommon::QueryUserPrivilege);
+}
+
+TOnlineAsyncOpHandle<FQueryUserPrivilege> FPrivilegesCommon::QueryUserPrivilege(FQueryUserPrivilege::Params&& Params)
+{
+	TOnlineAsyncOpRef<FQueryUserPrivilege> Operation = GetOp<FQueryUserPrivilege>(MoveTemp(Params));
+	Operation->SetError(Errors::NotImplemented());
+	return Operation->GetHandle();
+}
+
+/* UE::Online */}
\ No newline at end of file
diff --git a/Engine/Plugins/Online/Experimental/OnlineServices/Source/OnlineServicesCommon/Public/Online/ConnectivityCommon.h b/Engine/Plugins/Online/Experimental/OnlineServices/Source/OnlineServicesCommon/Public/Online/ConnectivityCommon.h
new file mode 100644
--- /dev/null
+++ b/Engine/Plugins/Online/Experimental/OnlineServices/Source/OnlineServicesCommon/Public/Online/ConnectivityCommon.h
@@ -0,0 +1,30 @@
+#pragma once
+
+#include "Online/Connectivity.h"
+#include "Online/OnlineComponent.h"
+
+namespace UE::Online {
+
+class FOnlineServicesCommon;
+
+class ONLINESERVICESCOMMON_API FConnectivityCommon : public TOnlineComponent<IConnectivity>
+{
+public:
+	using Super = IConnectivity;
+
+	FConnectivityCommon(FOnlineServicesCommon& InServices);
+
+	// TOnlineComponent
+	virtual void RegisterCommands() override;
+
+	// IConnectivity
+	virtual TOnlineResult<FGetConnectionStatus> GetConnectionStatus(FGetConnectionStatus::Params&& Params) override;
+	virtual TOnlineEvent<void(const FConnectionStatusChanged&)> OnConnectionStatusChanged() override;
+
+protected:
+	EOnlineServicesConnectionStatus ConnectionStatus = EOnlineServicesConnectionStatus::NotConnected;
+
+	TOnlineEventCallable<void(const FConnectionStatusChanged&)> OnConnectionStatusChangedEvent;
+};
+
+/* UE::Online */}
\ No newline at end of file
diff --git a/Engine/Plugins/Online/Experimental/OnlineServices/Source/OnlineServicesCommon/Public/Online/OnlineServicesCommon.h b/Engine/Plugins/Online/Experimental/OnlineServices/Source/OnlineServicesCommon/Public/Online/OnlineServicesCommon.h
--- a/Engine/Plugins/Online/Experimental/OnlineServices/Source/OnlineServicesCommon/Public/Online/OnlineServicesCommon.h
+++ b/Engine/Plugins/Online/Experimental/OnlineServices/Source/OnlineServicesCommon/Public/Online/OnlineServicesCommon.h
@@ -40,6 +40,8 @@ public:
 	virtual IPresencePtr GetPresenceInterface() override;
 	virtual IExternalUIPtr GetExternalUIInterface() override;
 	virtual ILobbiesPtr GetLobbiesInterface() override;
+	virtual IConnectivityPtr GetConnectivityInterface() override;
+	virtual IPrivilegesPtr GetPrivilegesInterface() override;
 
 	// FOnlineServicesCommon
 
diff --git a/Engine/Plugins/Online/Experimental/OnlineServices/Source/OnlineServicesCommon/Public/Online/PrivilegesCommon.h b/Engine/Plugins/Online/Experimental/OnlineServices/Source/OnlineServicesCommon/Public/Online/PrivilegesCommon.h
new file mode 100644
--- /dev/null
+++ b/Engine/Plugins/Online/Experimental/OnlineServices/Source/OnlineServicesCommon/Public/Online/PrivilegesCommon.h
@@ -0,0 +1,24 @@
+#pragma once
+
+#include "Online/Privileges.h"
+#include "Online/OnlineComponent.h"
+
+namespace UE::Online {
+
+class FOnlineServicesCommon;
+
+class ONLINESERVICESCOMMON_API FPrivilegesCommon : public TOnlineComponent<IPrivileges>
+{
+public:
+	using Super = IPrivileges;
+
+	FPrivilegesCommon(FOnlineServicesCommon& InServices);
+
+	// TOnlineComponent
+	virtual void RegisterCommands() override;
+
+	// IPermission
+	virtual TOnlineAsyncOpHandle<FQueryUserPrivilege> QueryUserPrivilege(FQueryUserPrivilege::Params&& Params) override;
+};
+
+/* UE::Online */}
\ No newline at end of file
diff --git a/Engine/Plugins/Online/Experimental/OnlineServices/Source/OnlineServicesInterface/Private/Connectivity.cpp b/Engine/Plugins/Online/Experimental/OnlineServices/Source/OnlineServicesInterface/Private/Connectivity.cpp
new file mode 100644
--- /dev/null
+++ b/Engine/Plugins/Online/Experimental/OnlineServices/Source/OnlineServicesInterface/Private/Connectivity.cpp
@@ -0,0 +1,34 @@
+// Copyright Epic Games, Inc. All Rights Reserved.
+
+#include "Online/Connectivity.h"
+
+namespace UE::Online {
+
+const TCHAR* LexToString(EOnlineServicesConnectionStatus Status)
+{
+	switch (Status)
+	{
+	case EOnlineServicesConnectionStatus::Connected:	return TEXT("Connected");
+	default:											checkNoEntry(); // Intentional fallthrough
+	case EOnlineServicesConnectionStatus::NotConnected:	return TEXT("NotConnected");
+	}
+}
+
+void LexFromString(EOnlineServicesConnectionStatus& OutStatus, const TCHAR* InStr)
+{
+	if (FCString::Stricmp(InStr, TEXT("Connected")) == 0)
+	{
+		OutStatus = EOnlineServicesConnectionStatus::Connected;
+	}
+	else if (FCString::Stricmp(InStr, TEXT("NotConnected")) == 0)
+	{
+		OutStatus = EOnlineServicesConnectionStatus::NotConnected;
+	}
+	else
+	{
+		checkNoEntry();
+		OutStatus = EOnlineServicesConnectionStatus::NotConnected;
+	}
+}
+
+/* UE::Online */}
\ No newline at end of file
diff --git a/Engine/Plugins/Online/Experimental/OnlineServices/Source/OnlineServicesInterface/Private/Privileges.cpp b/Engine/Plugins/Online/Experimental/OnlineServices/Source/OnlineServicesInterface/Private/Privileges.cpp
new file mode 100644
--- /dev/null
+++ b/Engine/Plugins/Online/Experimental/OnlineServices/Source/OnlineServicesInterface/Private/Privileges.cpp
@@ -0,0 +1,142 @@
+// Copyright Epic Games, Inc. All Rights Reserved.
+
+#include "Online/Privileges.h"
+
+namespace UE::Online {
+
+const TCHAR* LexToString(EUserPrivileges Type)
+{
+	switch (Type)
+	{
+
+	case EUserPrivileges::CanPlayOnline:				return TEXT("CanPlayOnline");
+	case EUserPrivileges::CanCommunicateViaTextOnline:	return TEXT("CanCommunicateViaTextOnline");
+	case EUserPrivileges::CanCommunicateViaVoiceOnline:	return TEXT("CanCommunicateViaVoiceOnline");
+	case EUserPrivileges::CanUseUserGeneratedContent:	return TEXT("CanUseUserGeneratedContent");
+	case EUserPrivileges::CanCrossPlay:					return TEXT("CanCrossPlay");
+	default:											checkNoEntry(); // Intentional fallthrough
+	case EUserPrivileges::CanPlay:						return TEXT("CanPlay");
+	}
+}
+
+void LexFromString(EUserPrivileges& OutType, const TCHAR* InStr)
+{
+	if (FCString::Stricmp(InStr, TEXT("CanPlay")) == 0)
+	{
+		OutType = EUserPrivileges::CanPlay;
+	}
+	else if (FCString::Stricmp(InStr, TEXT("CanPlayOnline")) == 0)
+	{
+		OutType = EUserPrivileges::CanPlayOnline;
+	}
+	else if (FCString::Stricmp(InStr, TEXT("CanCommunicateViaTextOnline")) == 0)
+	{
+		OutType = EUserPrivileges::CanCommunicateViaTextOnline;
+	}
+	else if (FCString::Stricmp(InStr, TEXT("CanCommunicateViaVoiceOnline")) == 0)
+	{
+		OutType = EUserPrivileges::CanCommunicateViaVoiceOnline;
+	}
+	else if (FCString::Stricmp(InStr, TEXT("CanUseUserGeneratedContent")) == 0)
+	{
+		OutType = EUserPrivileges::CanUseUserGeneratedContent;
+	}
+	else if (FCString::Stricmp(InStr, TEXT("CanUserCrossPlay")) == 0)
+	{
+		OutType = EUserPrivileges::CanCrossPlay;
+	}
+	else
+	{
+		checkNoEntry();
+		OutType = EUserPrivileges::CanPlay;
+	}
+}
+
+FString LexToString(EPrivilegeResults Type)
+{
+	if (Type == EPrivilegeResults::NoFailures)
+	{
+		return TEXT("NoFailures");
+	}
+	else
+	{
+		TArray<FString> ResultNames;
+
+		if (EnumHasAllFlags(Type, EPrivilegeResults::RequiredPatchAvailable))		ResultNames.Emplace(TEXT("RequiredPatchAvailable"));
+		if (EnumHasAllFlags(Type, EPrivilegeResults::RequiredSystemUpdate))			ResultNames.Emplace(TEXT("RequiredSystemUpdate"));
+		if (EnumHasAllFlags(Type, EPrivilegeResults::AgeRestrictionFailure))		ResultNames.Emplace(TEXT("AgeRestrictionFailure"));
+		if (EnumHasAllFlags(Type, EPrivilegeResults::AccountTypeFailure))			ResultNames.Emplace(TEXT("AccountTypeFailure"));
+		if (EnumHasAllFlags(Type, EPrivilegeResults::UserNotFound))					ResultNames.Emplace(TEXT("UserNotFound"));
+		if (EnumHasAllFlags(Type, EPrivilegeResults::UserNotLoggedIn))				ResultNames.Emplace(TEXT("UserNotLoggedIn"));
+		if (EnumHasAllFlags(Type, EPrivilegeResults::ChatRestriction))				ResultNames.Emplace(TEXT("ChatRestriction"));
+		if (EnumHasAllFlags(Type, EPrivilegeResults::UGCRestriction))				ResultNames.Emplace(TEXT("UGCRestriction"));
+		if (EnumHasAllFlags(Type, EPrivilegeResults::GenericFailure))				ResultNames.Emplace(TEXT("GenericFailure"));
+		if (EnumHasAllFlags(Type, EPrivilegeResults::OnlinePlayRestricted))			ResultNames.Emplace(TEXT("OnlinePlayRestricted"));
+		if (EnumHasAllFlags(Type, EPrivilegeResults::NetworkConnectionUnavailable))	ResultNames.Emplace(TEXT("NetworkConnectionUnavailable"));
+
+		return FString::Join(ResultNames, TEXT(" | "));
+	}
+}
+
+void LexFromString(EPrivilegeResults& OutType, const FString& InStr)
+{
+	if (FCString::Stricmp(*InStr, TEXT("NoFailures")) == 0)
+	{
+		OutType = EPrivilegeResults::NoFailures;
+	}
+	else
+	{
+		TArray<FString> PrivilegeNames;
+		InStr.ParseIntoArray(PrivilegeNames, TEXT(" | "));
+
+		for (const FString& PrivilegeName : PrivilegeNames)
+		{
+			if (FCString::Stricmp(*PrivilegeName, TEXT("RequiredPatchAvailable")) == 0)
+			{
+				OutType |= EPrivilegeResults::RequiredPatchAvailable;
+			}
+			else if (FCString::Stricmp(*PrivilegeName, TEXT("RequiredSystemUpdate")) == 0)
+			{
+				OutType |= EPrivilegeResults::RequiredSystemUpdate;
+			}
+			else if (FCString::Stricmp(*PrivilegeName, TEXT("AgeRestrictionFailure")) == 0)
+			{
+				OutType |= EPrivilegeResults::AgeRestrictionFailure;
+			}
+			else if (FCString::Stricmp(*PrivilegeName, TEXT("AccountTypeFailure")) == 0)
+			{
+				OutType |= EPrivilegeResults::AccountTypeFailure;
+			}
+			else if (FCString::Stricmp(*PrivilegeName, TEXT("UserNotFound")) == 0)
+			{
+				OutType |= EPrivilegeResults::UserNotFound;
+			}
+			else if (FCString::Stricmp(*PrivilegeName, TEXT("UserNotLoggedIn")) == 0)
+			{
+				OutType |= EPrivilegeResults::UserNotLoggedIn;
+			}
+			else if (FCString::Stricmp(*PrivilegeName, TEXT("ChatRestriction")) == 0)
+			{
+				OutType |= EPrivilegeResults::ChatRestriction;
+			}
+			else if (FCString::Stricmp(*PrivilegeName, TEXT("UGCRestriction")) == 0)
+			{
+				OutType |= EPrivilegeResults::UGCRestriction;
+			}
+			else if (FCString::Stricmp(*PrivilegeName, TEXT("GenericFailure")) == 0)
+			{
+				OutType |= EPrivilegeResults::GenericFailure;
+			}
+			else if (FCString::Stricmp(*PrivilegeName, TEXT("OnlinePlayRestricted")) == 0)
+			{
+				OutType |= EPrivilegeResults::OnlinePlayRestricted;
+			}
+			else if (FCString::Stricmp(*PrivilegeName, TEXT("NetworkConnectionUnavailable")) == 0)
+			{
+				OutType |= EPrivilegeResults::NetworkConnectionUnavailable;
+			}
+		}
+	}
+}
+
+/* UE::Online */}
\ No newline at end of file
diff --git a/Engine/Plugins/Online/Experimental/OnlineServices/Source/OnlineServicesInterface/Public/Online/Auth.h b/Engine/Plugins/Online/Experimental/OnlineServices/Source/OnlineServicesInterface/Public/Online/Auth.h
--- a/Engine/Plugins/Online/Experimental/OnlineServices/Source/OnlineServicesInterface/Public/Online/Auth.h
+++ b/Engine/Plugins/Online/Experimental/OnlineServices/Source/OnlineServicesInterface/Public/Online/Auth.h
@@ -218,6 +218,9 @@ BEGIN_ONLINE_STRUCT_META(FAuthGetAccountByAccountId::Result)
 END_ONLINE_STRUCT_META()
 
 BEGIN_ONLINE_STRUCT_META(FLoginStatusChanged)
+	ONLINE_STRUCT_FIELD(FLoginStatusChanged, LocalUserId),
+	ONLINE_STRUCT_FIELD(FLoginStatusChanged, PreviousStatus),
+	ONLINE_STRUCT_FIELD(FLoginStatusChanged, CurrentStatus)
 END_ONLINE_STRUCT_META()
 
 /* Meta*/ }
diff --git a/Engine/Plugins/Online/Experimental/OnlineServices/Source/OnlineServicesInterface/Public/Online/Connectivity.h b/Engine/Plugins/Online/Experimental/OnlineServices/Source/OnlineServicesInterface/Public/Online/Connectivity.h
new file mode 100644
--- /dev/null
+++ b/Engine/Plugins/Online/Experimental/OnlineServices/Source/OnlineServicesInterface/Public/Online/Connectivity.h
@@ -0,0 +1,83 @@
+// Copyright Epic Games, Inc. All Rights Reserved.
+
+#pragma once
+
+#include "Online/OnlineAsyncOpHandle.h"
+#include "Online/OnlineMeta.h"
+
+namespace UE::Online {
+
+enum class EOnlineServicesConnectionStatus : uint8
+{
+	/** Connected to the online servers */
+	Connected,
+	/** Successfully disconnected from the online servers */
+	NotConnected
+};
+ONLINESERVICESINTERFACE_API const TCHAR* LexToString(EOnlineServicesConnectionStatus Status);
+ONLINESERVICESINTERFACE_API void LexFromString(EOnlineServicesConnectionStatus& OutStatus, const TCHAR* InStr);
+
+struct FGetConnectionStatus
+{
+	static constexpr TCHAR Name[] = TEXT("GetConnectionStatus");
+
+	/** Input struct for Connectivity::FGetConnectionStatus */
+	struct Params
+	{
+		/** The name of the service being accessed (service dependent) */
+		FString ServiceName;
+	};
+
+	/** Output struct for Connectivity::FGetConnectionStatus */
+	struct Result
+	{
+		/** Connection status */
+		EOnlineServicesConnectionStatus Status;
+	};
+};
+
+/** Struct for ConnectionStatusChanged event */
+struct FConnectionStatusChanged
+{
+	/** The name of the service that is reporting (service dependent) */
+	FString ServiceName;
+	/** Previous login status */
+	EOnlineServicesConnectionStatus PreviousStatus;
+	/** Current login status */
+	EOnlineServicesConnectionStatus CurrentStatus;
+};
+
+class IConnectivity
+{
+public:
+	/**
+		* Accessor for the connection status
+		*/
+	virtual TOnlineResult<FGetConnectionStatus> GetConnectionStatus(FGetConnectionStatus::Params&& Params) = 0;
+
+	/**
+		* Event triggered when the connection status changes
+		*/
+	virtual TOnlineEvent<void(const FConnectionStatusChanged&)> OnConnectionStatusChanged() = 0;
+};
+
+namespace Meta {
+// TODO: Move to OnlineServices_Meta.inl file?
+
+BEGIN_ONLINE_STRUCT_META(FGetConnectionStatus::Params)
+	ONLINE_STRUCT_FIELD(FGetConnectionStatus::Params, ServiceName)
+END_ONLINE_STRUCT_META()
+
+BEGIN_ONLINE_STRUCT_META(FGetConnectionStatus::Result)
+	ONLINE_STRUCT_FIELD(FGetConnectionStatus::Result, Status)
+END_ONLINE_STRUCT_META()
+
+BEGIN_ONLINE_STRUCT_META(FConnectionStatusChanged)
+	ONLINE_STRUCT_FIELD(FConnectionStatusChanged, ServiceName),
+	ONLINE_STRUCT_FIELD(FConnectionStatusChanged, PreviousStatus),
+	ONLINE_STRUCT_FIELD(FConnectionStatusChanged, CurrentStatus)
+END_ONLINE_STRUCT_META()
+
+/* Meta*/}
+
+/* UE::Online */}
\ No newline at end of file
diff --git a/Engine/Plugins/Online/Experimental/OnlineServices/Source/OnlineServicesInterface/Public/Online/OnlineServices.h b/Engine/Plugins/Online/Experimental/OnlineServices/Source/OnlineServicesInterface/Public/Online/OnlineServices.h
--- a/Engine/Plugins/Online/Experimental/OnlineServices/Source/OnlineServicesInterface/Public/Online/OnlineServices.h
+++ b/Engine/Plugins/Online/Experimental/OnlineServices/Source/OnlineServicesInterface/Public/Online/OnlineServices.h
@@ -16,6 +16,8 @@ using IFriendsPtr = TSharedPtr<class IFriends>;
 using IPresencePtr = TSharedPtr<class IPresence>;
 using IExternalUIPtr = TSharedPtr<class IExternalUI>;
 using ILobbiesPtr = TSharedPtr<class ILobbies>;
+using IConnectivityPtr = TSharedPtr<class IConnectivity>;
+using IPrivilegesPtr = TSharedPtr<class IPrivileges>;
 
 class ONLINESERVICESINTERFACE_API IOnlineServices
 {
@@ -55,6 +57,16 @@ public:
 	 * @return lobbies implementation, may be null if not implemented for this service
 	 */
 	virtual ILobbiesPtr GetLobbiesInterface() = 0;
+
+	/**
+	 * 
+	 */
+	virtual IConnectivityPtr GetConnectivityInterface() = 0;
+
+	/**
+	 *
+	 */
+	virtual IPrivilegesPtr GetPrivilegesInterface() = 0;
 };
 
 /**
diff --git a/Engine/Plugins/Online/Experimental/OnlineServices/Source/OnlineServicesInterface/Public/Online/Privileges.h b/Engine/Plugins/Online/Experimental/OnlineServices/Source/OnlineServicesInterface/Public/Online/Privileges.h
new file mode 100644
--- /dev/null
+++ b/Engine/Plugins/Online/Experimental/OnlineServices/Source/OnlineServicesInterface/Public/Online/Privileges.h
@@ -0,0 +1,103 @@
+// Copyright Epic Games, Inc. All Rights Reserved.
+
+#pragma once
+
+#include "Online/CoreOnline.h"
+#include "Online/OnlineAsyncOpHandle.h"
+#include "Online/OnlineMeta.h"
+
+namespace UE::Online {
+
+enum class EUserPrivileges : uint8
+{
+	/** Whether the user can play at all, online or offline*/
+	CanPlay,
+	/** Whether the user can play in online modes */
+	CanPlayOnline,
+	/** Whether the user can use text chat */
+	CanCommunicateViaTextOnline,
+	/** Whether the user can use voice chat */
+	CanCommunicateViaVoiceOnline,
+	/** Whether the user can use content generated by other users */
+	CanUseUserGeneratedContent,
+	/** Whether the user can ever participate in cross-play*/
+	CanCrossPlay
+};
+ONLINESERVICESINTERFACE_API const TCHAR* LexToString(EUserPrivileges Type);
+ONLINESERVICESINTERFACE_API void LexFromString(EUserPrivileges& OutType, const TCHAR* InStr);
+
+enum class EPrivilegeResults : uint32
+{
+	/** The user has the requested privilege */
+	NoFailures = 0,
+	/** Patch required before the user can use the privilege */
+	RequiredPatchAvailable = 1 << 0,
+	/** System update required before the user can use the privilege */
+	RequiredSystemUpdate = 1 << 1,
+	/** Parental control failure usually */
+	AgeRestrictionFailure = 1 << 2,
+	/** Premium account required */
+	AccountTypeFailure = 1 << 3,
+	/** Invalid user */
+	UserNotFound = 1 << 4,
+	/** User must be logged in */
+	UserNotLoggedIn = 1 << 5,
+	/** User restricted from chat */
+	ChatRestriction = 1 << 6,
+	/** User restricted from User Generated Content */
+	UGCRestriction = 1 << 7,
+	/** Platform failed for unknown reason and handles its own dialogs */
+	GenericFailure = 1 << 8,
+	/** Online play is restricted */
+	OnlinePlayRestricted = 1 << 9,
+	/** Check failed because network is unavailable */
+	NetworkConnectionUnavailable = 1 << 10,
+};
+ENUM_CLASS_FLAGS(EPrivilegeResults);
+ONLINESERVICESINTERFACE_API FString LexToString(EPrivilegeResults Type);
+ONLINESERVICESINTERFACE_API void LexFromString(EPrivilegeResults& OutType, const FString& InStr);
+
+struct FQueryUserPrivilege
+{
+	static constexpr TCHAR Name[] = TEXT("QueryUserPrivilege");
+
+	/** Input struct for Misc::FQueryUserPrivilege */
+	struct Params
+	{
+		/** User id whose privilege we want to query */
+		FOnlineAccountIdHandle LocalUserId;
+		/** Privilege to query */
+		EUserPrivileges Privilege;
+	};
+
+	/** Output struct for Misc::FQueryUserPrivilege */
+	struct Result
+	{
+		/** Privilege result, bitwise OR of various EPrivilegeResults values */
+		EPrivilegeResults PrivilegeResult = EPrivilegeResults::NoFailures;
+	};
+};
+
+class IPrivileges
+{
+	/**
+		* Queries about a given privilege for a user
+		*/
+	virtual TOnlineAsyncOpHandle<FQueryUserPrivilege> QueryUserPrivilege(FQueryUserPrivilege::Params&& Params) = 0;
+};
+
+namespace Meta {
+// TODO: Move to OnlineServices_Meta.inl file?
+
+BEGIN_ONLINE_STRUCT_META(FQueryUserPrivilege::Params)
+	ONLINE_STRUCT_FIELD(FQueryUserPrivilege::Params, LocalUserId),
+	ONLINE_STRUCT_FIELD(FQueryUserPrivilege::Params, Privilege)
+END_ONLINE_STRUCT_META()
+
+BEGIN_ONLINE_STRUCT_META(FQueryUserPrivilege::Result)
+	ONLINE_STRUCT_FIELD(FQueryUserPrivilege::Result, PrivilegeResult)
+END_ONLINE_STRUCT_META()
+
+/* Meta*/}
+
+/* UE::Online */}
\ No newline at end of file
