- Create Default and Platform contexts

diff --git a/Engine/Plugins/Online/Experimental/OnlineServices/Source/OnlineServicesInterface/Private/OnlineServicesRegistry.cpp b/Engine/Plugins/Online/Experimental/OnlineServices/Source/OnlineServicesInterface/Private/OnlineServicesRegistry.cpp
--- a/Engine/Plugins/Online/Experimental/OnlineServices/Source/OnlineServicesInterface/Private/OnlineServicesRegistry.cpp
+++ b/Engine/Plugins/Online/Experimental/OnlineServices/Source/OnlineServicesInterface/Private/OnlineServicesRegistry.cpp
@@ -3,6 +3,7 @@
 #include "Online/OnlineServicesRegistry.h"
 #include "Online/OnlineServicesDelegates.h"
 
+#include "Misc/ConfigCacheIni.h"
 #include "Misc/LazySingleton.h"
 
 namespace UE::Online {
@@ -49,17 +50,33 @@ TSharedPtr<IOnlineServices> FOnlineServicesRegistry::GetNamedServicesInstance(EO
 {
 	TSharedPtr<IOnlineServices> Services;
 
-	if (TSharedRef<IOnlineServices>* ServicesPtr = NamedServiceInstances.FindOrAdd(OnlineServices).Find(InstanceName))
+	if (OnlineServices == EOnlineServices::Default)
+	{
+		FString Value;
+		GConfig->GetString(TEXT("OnlineServices"), TEXT("DefaultServices"), Value, GEngineIni);
+		LexFromString(OnlineServices, *Value);
+	}
+	else if (OnlineServices == EOnlineServices::Platform)
 	{
-		Services = *ServicesPtr;
+		FString Value;
+		GConfig->GetString(TEXT("OnlineServices"), TEXT("PlatformServices"), Value, GEngineIni);
+		LexFromString(OnlineServices, *Value);
 	}
-	else
+
+	if (OnlineServices < EOnlineServices::None)
 	{
-		Services = CreateServices(OnlineServices);
-		if (Services.IsValid())
+		if (TSharedRef<IOnlineServices>* ServicesPtr = NamedServiceInstances.FindOrAdd(OnlineServices).Find(InstanceName))
+		{
+			Services = *ServicesPtr;
+		}
+		else
 		{
-			NamedServiceInstances.FindOrAdd(OnlineServices).Add(InstanceName, Services.ToSharedRef());
-			OnOnlineServicesCreated.Broadcast(Services.ToSharedRef());
+			Services = CreateServices(OnlineServices);
+			if (Services.IsValid())
+			{
+				NamedServiceInstances.FindOrAdd(OnlineServices).Add(InstanceName, Services.ToSharedRef());
+				OnOnlineServicesCreated.Broadcast(Services.ToSharedRef());
+			}
 		}
 	}
 
diff --git a/Engine/Plugins/Online/Experimental/OnlineServicesEOS/Source/Private/OnlineServicesEOS.cpp b/Engine/Plugins/Online/Experimental/OnlineServicesEOS/Source/Private/OnlineServicesEOS.cpp
--- a/Engine/Plugins/Online/Experimental/OnlineServicesEOS/Source/Private/OnlineServicesEOS.cpp
+++ b/Engine/Plugins/Online/Experimental/OnlineServicesEOS/Source/Private/OnlineServicesEOS.cpp
@@ -62,16 +62,9 @@ void FOnlineServicesEOS::Initialize()
 	FTCHARToUTF8 ClientId(*EOSPlatformConfig.ClientId);
 	FTCHARToUTF8 ClientSecret(*EOSPlatformConfig.ClientSecret);
 
-	// Init EOS SDK
-	static struct FReservedOptions
-	{
-		int32_t ApiVersion;
-		const char* BackendEnvironment;
-	} ReservedOptions = { 1, "GameDev" };
-
 	EOS_Platform_Options PlatformOptions = {};
-	PlatformOptions.Reserved = &ReservedOptions;
 	PlatformOptions.ApiVersion = EOS_PLATFORM_OPTIONS_API_LATEST;
+	PlatformOptions.Reserved = nullptr;
 	PlatformOptions.bIsServer = EOS_FALSE;
 	PlatformOptions.OverrideCountryCode = nullptr;
 	PlatformOptions.OverrideLocaleCode = nullptr;
diff --git a/Engine/Source/Runtime/CoreOnline/Private/Online/CoreOnline.cpp b/Engine/Source/Runtime/CoreOnline/Private/Online/CoreOnline.cpp
--- a/Engine/Source/Runtime/CoreOnline/Private/Online/CoreOnline.cpp
+++ b/Engine/Source/Runtime/CoreOnline/Private/Online/CoreOnline.cpp
@@ -43,8 +43,9 @@ const TCHAR* LexToString(EOnlineServices Value)
 	case EOnlineServices::GameDefined_1:	return TEXT("GameDefined_1");
 	case EOnlineServices::GameDefined_2:	return TEXT("GameDefined_2");
 	case EOnlineServices::GameDefined_3:	return TEXT("GameDefined_3");
-	case EOnlineServices::Platform:			return TEXT("Platform");
 	default:								checkNoEntry(); // Intentional fallthrough
+	case EOnlineServices::None:				return TEXT("None");
+	case EOnlineServices::Platform:			return TEXT("Platform");
 	case EOnlineServices::Default:			return TEXT("Default");
 	}
 }
@@ -179,6 +180,10 @@ void LexFromString(EOnlineServices& OutValue, const TCHAR* InStr)
 	{
 		OutValue = EOnlineServices::GameDefined_3;
 	}
+	else if (FCString::Stricmp(InStr, TEXT("None")) == 0)
+	{
+		OutValue = EOnlineServices::None;
+	}
 	else if (FCString::Stricmp(InStr, TEXT("Platform")) == 0)
 	{
 		OutValue = EOnlineServices::Platform;
@@ -189,8 +194,9 @@ void LexFromString(EOnlineServices& OutValue, const TCHAR* InStr)
 	}
 	else
 	{
-		checkNoEntry();
-		OutValue = EOnlineServices::Default;
+		// Empty string is ok and gets mapped to "None". Anything else implies we are missing a case.
+		check(FCString::Strlen(InStr) == 0);
+		OutValue = EOnlineServices::None;
 	}
 }
 
diff --git a/Engine/Source/Runtime/CoreOnline/Public/Online/CoreOnline.h b/Engine/Source/Runtime/CoreOnline/Public/Online/CoreOnline.h
--- a/Engine/Source/Runtime/CoreOnline/Public/Online/CoreOnline.h
+++ b/Engine/Source/Runtime/CoreOnline/Public/Online/CoreOnline.h
@@ -250,6 +250,8 @@ enum class EOnlineServices : uint8
 	GameDefined_1,
 	GameDefined_2,
 	GameDefined_3,
+	// None, used internally to resolve Platform or Default if they are not configured
+	None = 253,
 	// Platform native, may not exist for all platforms
 	Platform = 254,
 	// Default, configured via ini, TODO: List specific ini section/key
