- Add support for OSS V2 connection status

diff --git a/Engine/Plugins/Online/Experimental/OnlineServices/Source/OnlineServicesInterface/Public/Online/Connectivity.h b/Engine/Plugins/Online/Experimental/OnlineServices/Source/OnlineServicesInterface/Public/Online/Connectivity.h
--- a/Engine/Plugins/Online/Experimental/OnlineServices/Source/OnlineServicesInterface/Public/Online/Connectivity.h
+++ b/Engine/Plugins/Online/Experimental/OnlineServices/Source/OnlineServicesInterface/Public/Online/Connectivity.h
@@ -24,7 +24,7 @@ struct FGetConnectionStatus
 	/** Input struct for Connectivity::FGetConnectionStatus */
 	struct Params
 	{
-		/** The name of the service being accessed (service dependent) */
+		/** Optional. If set, the name of the specific service being accessed (service dependent). */
 		FString ServiceName;
 	};
 
@@ -39,11 +39,11 @@ struct FGetConnectionStatus
 /** Struct for ConnectionStatusChanged event */
 struct FConnectionStatusChanged
 {
-	/** The name of the service that is reporting (service dependent) */
+	/** The name of the service that is reporting (service dependent). Can be empty. */
 	FString ServiceName;
-	/** Previous login status */
+	/** Previous connection status */
 	EOnlineServicesConnectionStatus PreviousStatus;
-	/** Current login status */
+	/** Current connection status */
 	EOnlineServicesConnectionStatus CurrentStatus;
 };
 
