Add unit tests for client lobby data.
Move common lobby types into LobbiesCommonTypes.h / cpp
Move lobbies functional tests into LobbiesCommonTests.h / cpp

diff --git a/Engine/Plugins/Online/Experimental/OnlineServices/Source/OnlineServicesCommon/Private/Online/LobbiesCommon.cpp b/Engine/Plugins/Online/Experimental/OnlineServices/Source/OnlineServicesCommon/Private/Online/LobbiesCommon.cpp
--- a/Engine/Plugins/Online/Experimental/OnlineServices/Source/OnlineServicesCommon/Private/Online/LobbiesCommon.cpp
+++ b/Engine/Plugins/Online/Experimental/OnlineServices/Source/OnlineServicesCommon/Private/Online/LobbiesCommon.cpp
@@ -4,109 +4,10 @@
 
 #include "Online/Auth.h"
 #include "Online/OnlineServicesCommon.h"
+#include "Online/LobbiesCommonTests.h"
 
 namespace UE::Online {
 
-template <typename DataType, typename OpType>
-const DataType& GetOpDataChecked(const TOnlineAsyncOp<OpType>& Op, const FString& Key)
-{
-	const DataType* Data = Op.Data.template Get<DataType>(Key);
-	check(Data);
-	return *Data;
-}
-
-TSharedPtr<FLobbySchema> FLobbySchema::Create(FLobbySchemaConfig LobbySchemaConfig)
-{
-	TSharedRef<FLobbySchema> Schema = MakeShared<FLobbySchema>();
-	Schema->SchemaId = LobbySchemaConfig.SchemaName;
-
-	// Todo: for real validation and init
-	if (Schema->SchemaId == FLobbySchemaId())
-	{
-		return TSharedPtr<FLobbySchema>();
-	}
-	else
-	{
-		return Schema;
-	}
-}
-
-TSharedRef<FLobbyServiceAttributeChanges> FLobbySchema::TranslateLobbyAttributes(const TMap<FLobbyAttributeId, FLobbyVariant>& LobbyAttributes, const FLobbyClientAttributeChanges& ClientAttributeChanges) const
-{
-	// Todo: implementation.
-
-	TSharedRef<FLobbyServiceAttributeChanges> Changes = MakeShared<FLobbyServiceAttributeChanges>();
-	Changes->MutatedAttributes = ClientAttributeChanges.MutatedAttributes;
-	Changes->ClearedAttributes = ClientAttributeChanges.ClearedAttributes;
-	return Changes;
-}
-
-TSharedRef<FLobbyClientAttributeChanges> FLobbySchema::TranslateLobbyAttributes(const FLobbyServiceAttributeChanges& ServiceAttributeChanges) const
-{
-	// Todo: implementation.
-
-	TSharedRef<FLobbyClientAttributeChanges> Changes = MakeShared<FLobbyClientAttributeChanges>();
-	Changes->MutatedAttributes = ServiceAttributeChanges.MutatedAttributes;
-	Changes->ClearedAttributes = ServiceAttributeChanges.ClearedAttributes;
-	return Changes;
-}
-
-TSharedRef<FLobbyServiceAttributeChanges> FLobbySchema::TranslateLobbyMemberAttributes(const TMap<FLobbyAttributeId, FLobbyVariant>& LobbyAttributes, const FLobbyClientAttributeChanges& ClientAttributeChanges) const
-{
-	// Todo: implementation.
-
-	TSharedRef<FLobbyServiceAttributeChanges> Changes = MakeShared<FLobbyServiceAttributeChanges>();
-	Changes->MutatedAttributes = ClientAttributeChanges.MutatedAttributes;
-	Changes->ClearedAttributes = ClientAttributeChanges.ClearedAttributes;
-	return Changes;
-}
-
-TSharedRef<FLobbyClientAttributeChanges> FLobbySchema::TranslateLobbyMemberAttributes(const FLobbyServiceAttributeChanges& ServiceAttributeChanges) const
-{
-	// Todo: implementation.
-
-	TSharedRef<FLobbyClientAttributeChanges> Changes = MakeShared<FLobbyClientAttributeChanges>();
-	Changes->MutatedAttributes = ServiceAttributeChanges.MutatedAttributes;
-	Changes->ClearedAttributes = ServiceAttributeChanges.ClearedAttributes;
-	return Changes;
-}
-
-bool FLobbySchemaRegistry::Initialize(TArray<FLobbySchemaConfig> LobbySchemaConfigs)
-{
-	bool bSuccess = true;
-
-	for (FLobbySchemaConfig& LobbySchemaConfig : LobbySchemaConfigs)
-	{
-		bSuccess &= RegisterSchema(MoveTemp(LobbySchemaConfig));
-	}
-
-	return bSuccess;
-}
-
-TSharedPtr<FLobbySchema> FLobbySchemaRegistry::FindSchema(FLobbySchemaId SchemaId)
-{
-	TSharedRef<FLobbySchema>* Schema = RegisteredSchemas.Find(SchemaId);
-	return Schema ? *Schema : TSharedPtr<FLobbySchema>();
-}
-
-bool FLobbySchemaRegistry::RegisterSchema(FLobbySchemaConfig LobbySchemaConfig)
-{
-	// Todo: handle schema hierarchy.
-	// Todo: handle schema parsing.
-
-	FLobbySchemaId SchemaName = LobbySchemaConfig.SchemaName;
-
-	if (TSharedPtr<FLobbySchema> Schema = FLobbySchema::Create(MoveTemp(LobbySchemaConfig)))
-	{
-		RegisteredSchemas.Emplace(SchemaName, Schema.ToSharedRef());
-		return true;
-	}
-	else
-	{
-		return false;
-	}
-}
-
 FLobbiesCommon::FLobbiesCommon(FOnlineServicesCommon& InServices)
 	: TOnlineComponent(TEXT("Lobbies"), InServices)
 {
@@ -167,7 +68,7 @@ void FLobbiesCommon::RegisterCommands()
 
 #if LOBBIES_FUNCTIONAL_TEST_ENABLED
 	RegisterCommand(&FLobbiesCommon::FunctionalTest);
-#endif
+#endif // LOBBIES_FUNCTIONAL_TEST_ENABLED
 }
 
 TOnlineAsyncOpHandle<FCreateLobby> FLobbiesCommon::CreateLobby(FCreateLobby::Params&& Params)
@@ -321,804 +222,19 @@ TOnlineEvent<void(const FLobbyInvitationRemoved&)> FLobbiesCommon::OnLobbyInvita
 	return LobbyEvents.OnLobbyInvitationRemoved;
 }
 
-// todo: TFuture<void> does not work as expected with "Then".
-TFuture<int> AwaitNextGameTick()
-{
-	TSharedRef<TPromise<int>> Promise = MakeShared<TPromise<int>>();
-	TFuture<int> Future = Promise->GetFuture();
-
-	FTSTicker::GetCoreTicker().AddTicker(
-	TEXT("FLobbiesCommon::AwaitNextGameTick"),
-	0.f,
-	[Promise](float)
-	{
-		Promise->EmplaceValue();
-		return false;
-	});
-
-	return Future;
-}
-
-TFuture<TDefaultErrorResultInternal<FOnlineLobbyIdHandle>> FLobbiesCommon::AwaitInvitation(
-	FOnlineAccountIdHandle TargetAccountId,
-	FOnlineLobbyIdHandle LobbyId,
-	float TimeoutSeconds)
-{
-	// todo: Make this nicer - current implementation has issue on shutdown.
-	struct AwaitInvitationState
-	{
-		TPromise<TDefaultErrorResultInternal<FOnlineLobbyIdHandle>> Promise;
-		FOnlineEventDelegateHandle OnLobbyInvitationAddedHandle;
-		FTSTicker::FDelegateHandle OnAwaitExpiredHandle;
-		bool bSignaled = false;
-	};
-
-	TSharedRef<AwaitInvitationState> AwaitState = MakeShared<AwaitInvitationState>();
-
-	AwaitState->OnLobbyInvitationAddedHandle = OnLobbyInvitationAdded().Add(
-	[this, TargetAccountId, LobbyId, AwaitState](const FLobbyInvitationAdded& Notification)
-	{
-		FOnlineEventDelegateHandle DelegateHandle = MoveTemp(AwaitState->OnLobbyInvitationAddedHandle);
-
-		if (!AwaitState->bSignaled)
-		{
-			if (Notification.Lobby->LobbyId == LobbyId && Notification.LocalUserId == TargetAccountId)
-			{
-				FTSTicker::GetCoreTicker().RemoveTicker(AwaitState->OnAwaitExpiredHandle);
-				AwaitState->bSignaled = true;
-
-				AwaitNextGameTick()
-				.Then([AwaitState, LobbyId](TFuture<int>&&)
-				{
-					AwaitState->Promise.EmplaceValue(LobbyId);
-				});
-			}
-		}
-	});
-
-	AwaitState->OnAwaitExpiredHandle = FTSTicker::GetCoreTicker().AddTicker(
-	TEXT("FLobbiesCommon::AwaitInvitation"),
-	TimeoutSeconds,
-	[this, AwaitState](float)
-	{
-		if (!AwaitState->bSignaled)
-		{
-			// Todo: Errors.
-			AwaitState->bSignaled = true;
-			AwaitState->OnLobbyInvitationAddedHandle.Unbind();
-			AwaitState->Promise.EmplaceValue(Errors::NotImplemented());
-			FTSTicker::GetCoreTicker().RemoveTicker(AwaitState->OnAwaitExpiredHandle);
-		}
-
-		return false;
-	});
-
-	return AwaitState->Promise.GetFuture();
-}
-
 #if LOBBIES_FUNCTIONAL_TEST_ENABLED
-struct FFunctionalTestConfig
-{
-	FString TestAccount1Type;
-	FString TestAccount1Id;
-	FString TestAccount1Token;
-
-	FString TestAccount2Type;
-	FString TestAccount2Id;
-	FString TestAccount2Token;
-
-	float InvitationWaitSeconds = 10.f;
-};
-
-namespace Meta {
-
-BEGIN_ONLINE_STRUCT_META(FFunctionalTestConfig)
-	ONLINE_STRUCT_FIELD(FFunctionalTestConfig, TestAccount1Type),
-	ONLINE_STRUCT_FIELD(FFunctionalTestConfig, TestAccount1Id),
-	ONLINE_STRUCT_FIELD(FFunctionalTestConfig, TestAccount1Token),
-
-	ONLINE_STRUCT_FIELD(FFunctionalTestConfig, TestAccount2Type),
-	ONLINE_STRUCT_FIELD(FFunctionalTestConfig, TestAccount2Id),
-	ONLINE_STRUCT_FIELD(FFunctionalTestConfig, TestAccount2Token)
-END_ONLINE_STRUCT_META()
-
-/* Meta */ }
-
-class FLobbyEventCapture final
-{
-public:
-	UE_NONCOPYABLE(FLobbyEventCapture);
-
-	FLobbyEventCapture(FLobbiesCommon& Lobbies)
-		: LobbyJoinedHandle(Lobbies.OnLobbyJoined().Add(this, &FLobbyEventCapture::OnLobbyJoined))
-		, LobbyLeftHandle(Lobbies.OnLobbyLeft().Add(this, &FLobbyEventCapture::OnLobbyLeft))
-		, LobbyMemberJoinedHandle(Lobbies.OnLobbyMemberJoined().Add(this, &FLobbyEventCapture::OnLobbyMemberJoined))
-		, LobbyMemberLeftHandle(Lobbies.OnLobbyMemberLeft().Add(this, &FLobbyEventCapture::OnLobbyMemberLeft))
-		, LobbyLeaderChangedHandle(Lobbies.OnLobbyLeaderChanged().Add(this, &FLobbyEventCapture::OnLobbyLeaderChanged))
-		, LobbySchemaChangedHandle(Lobbies.OnLobbySchemaChanged().Add(this, &FLobbyEventCapture::OnLobbySchemaChanged))
-		, LobbyAttributesChangedHandle(Lobbies.OnLobbyAttributesChanged().Add(this, &FLobbyEventCapture::OnLobbyAttributesChanged))
-		, LobbyMemberAttributesChangedHandle(Lobbies.OnLobbyMemberAttributesChanged().Add(this, &FLobbyEventCapture::OnLobbyMemberAttributesChanged))
-		, LobbyInvitationAddedHandle(Lobbies.OnLobbyInvitationAdded().Add(this, &FLobbyEventCapture::OnLobbyInvitationAdded))
-		, LobbyInvitationRemovedHandle(Lobbies.OnLobbyInvitationRemoved().Add(this, &FLobbyEventCapture::OnLobbyInvitationRemoved))
-	{
-	}
-
-	void Empty()
-	{
-		LobbyJoined.Empty();
-		LobbyLeft.Empty();
-		LobbyMemberJoined.Empty();
-		LobbyMemberLeft.Empty();
-		LobbyLeaderChanged.Empty();
-		LobbySchemaChanged.Empty();
-		LobbyAttributesChanged.Empty();
-		LobbyMemberAttributesChanged.Empty();
-		LobbyInvitationAdded.Empty();
-		LobbyInvitationRemoved.Empty();
-		NextIndex = 0;
-		TotalNotificationsReceived = 0;
-	}
-
-	uint32 GetTotalNotificationsReceived() const
-	{
-		return TotalNotificationsReceived;
-	}
-
-	template <typename NotificationType>
-	struct TNotificationInfo
-	{
-		NotificationType Notification;
-		int32 GlobalIndex = 0;
-	};
-
-	TArray<TNotificationInfo<FLobbyJoined>> LobbyJoined;
-	TArray<TNotificationInfo<FLobbyLeft>> LobbyLeft;
-	TArray<TNotificationInfo<FLobbyMemberJoined>> LobbyMemberJoined;
-	TArray<TNotificationInfo<FLobbyMemberLeft>> LobbyMemberLeft;
-	TArray<TNotificationInfo<FLobbyLeaderChanged>> LobbyLeaderChanged;
-	TArray<TNotificationInfo<FLobbySchemaChanged>> LobbySchemaChanged;
-	TArray<TNotificationInfo<FLobbyAttributesChanged>> LobbyAttributesChanged;
-	TArray<TNotificationInfo<FLobbyMemberAttributesChanged>> LobbyMemberAttributesChanged;
-	TArray<TNotificationInfo<FLobbyInvitationAdded>> LobbyInvitationAdded;
-	TArray<TNotificationInfo<FLobbyInvitationRemoved>> LobbyInvitationRemoved;
-
-private:
-	template <typename ContainerType, typename NotificationType>
-	void AddEvent(ContainerType& Container, const NotificationType& Notification)
-	{
-		Container.Add({Notification, NextIndex++});
-		++TotalNotificationsReceived;
-	}
-
-	void OnLobbyJoined(const FLobbyJoined& Notification) { AddEvent(LobbyJoined, Notification); }
-	void OnLobbyLeft(const FLobbyLeft& Notification) { AddEvent(LobbyLeft, Notification); }
-	void OnLobbyMemberJoined(const FLobbyMemberJoined& Notification) { AddEvent(LobbyMemberJoined, Notification); }
-	void OnLobbyMemberLeft(const FLobbyMemberLeft& Notification) { AddEvent(LobbyMemberLeft, Notification); }
-	void OnLobbyLeaderChanged(const FLobbyLeaderChanged& Notification) { AddEvent(LobbyLeaderChanged, Notification); }
-	void OnLobbySchemaChanged(const FLobbySchemaChanged& Notification) { AddEvent(LobbySchemaChanged, Notification); }
-	void OnLobbyAttributesChanged(const FLobbyAttributesChanged& Notification) { AddEvent(LobbyAttributesChanged, Notification); }
-	void OnLobbyMemberAttributesChanged(const FLobbyMemberAttributesChanged& Notification) { AddEvent(LobbyMemberAttributesChanged, Notification); }
-	void OnLobbyInvitationAdded(const FLobbyInvitationAdded& Notification) { AddEvent(LobbyInvitationAdded, Notification); }
-	void OnLobbyInvitationRemoved(const FLobbyInvitationRemoved& Notification) { AddEvent(LobbyInvitationRemoved, Notification); }
-
-	FOnlineEventDelegateHandle LobbyJoinedHandle;
-	FOnlineEventDelegateHandle LobbyLeftHandle;
-	FOnlineEventDelegateHandle LobbyMemberJoinedHandle;
-	FOnlineEventDelegateHandle LobbyMemberLeftHandle;
-	FOnlineEventDelegateHandle LobbyLeaderChangedHandle;
-	FOnlineEventDelegateHandle LobbySchemaChangedHandle;
-	FOnlineEventDelegateHandle LobbyAttributesChangedHandle;
-	FOnlineEventDelegateHandle LobbyMemberAttributesChangedHandle;
-	FOnlineEventDelegateHandle LobbyInvitationAddedHandle;
-	FOnlineEventDelegateHandle LobbyInvitationRemovedHandle;
-	int32 NextIndex = 0;
-	uint32 TotalNotificationsReceived = 0;
-};
-
-static const FString ConfigName = TEXT("FunctionalTest");
-static const FString User1KeyName = TEXT("User1");
-static const FString User2KeyName = TEXT("User2");
-static const FString CreateLobbyKeyName = TEXT("CreateLobby");
-static const FString FindLobbyKeyName = TEXT("FindLobby");
-static const FString LobbyEventCaptureKeyName = TEXT("LobbyEventCapture");
-static const FString ConfigNameKeyName = TEXT("Config");
-
-TOnlineAsyncOpHandle<FFunctionalTestLobbies> FLobbiesCommon::FunctionalTest(FFunctionalTestLobbies::Params&& InParams)
+TOnlineAsyncOpHandle<FFunctionalTestLobbies> FLobbiesCommon::FunctionalTest(FFunctionalTestLobbies::Params&& Params)
 {
-	TSharedRef<FFunctionalTestConfig> TestConfig = MakeShared<FFunctionalTestConfig>();
-	LoadConfig(*TestConfig, ConfigName);
-
-	TOnlineAsyncOpRef<FFunctionalTestLobbies> Op = GetOp<FFunctionalTestLobbies>(MoveTemp(InParams));
-
-	// Set up event capturing.
-	Op->Data.Set(LobbyEventCaptureKeyName, MakeShared<FLobbyEventCapture>(*this));
-
-	Op->Then(ConsumeStepResult<FFunctionalTestLogoutAllUsers>(*Op, &FLobbiesCommon::FunctionalTestLogoutAllUsers, FFunctionalTestLogoutAllUsers::Params{}))
-	.Then(CaptureStepResult<FFunctionalTestLoginUser>(*Op, User1KeyName, &FLobbiesCommon::FunctionalTestLoginUser, FFunctionalTestLoginUser::Params{FPlatformMisc::GetPlatformUserForUserIndex(0), TestConfig->TestAccount1Type, TestConfig->TestAccount1Id, TestConfig->TestAccount1Token}))
-	.Then(CaptureStepResult<FFunctionalTestLoginUser>(*Op, User2KeyName, &FLobbiesCommon::FunctionalTestLoginUser, FFunctionalTestLoginUser::Params{FPlatformMisc::GetPlatformUserForUserIndex(1), TestConfig->TestAccount2Type, TestConfig->TestAccount2Id, TestConfig->TestAccount2Token}))
-
-	//----------------------------------------------------------------------------------------------
-	// Test 1:
-	//    Step 1: Create a lobby with primary user.
-	//    Step 2: Modify lobby attribute.
-	//    Step 3: Modify lobby member attribute.
-	//    Step 4: Leave lobby with primary user.
-	//----------------------------------------------------------------------------------------------
-
-	.Then([this](TOnlineAsyncOp<FFunctionalTestLobbies>& InAsyncOp)
-	{
-		// Prepare for next test check.
-		GetOpDataChecked<TSharedRef<FLobbyEventCapture>>(InAsyncOp, LobbyEventCaptureKeyName)->Empty();
-
-		const FFunctionalTestLoginUser::Result& User1Info = GetOpDataChecked<FFunctionalTestLoginUser::Result>(InAsyncOp, User1KeyName);
-
-		// Create a lobby
-		FCreateLobby::Params Params;
-		Params.LocalUserId = User1Info.AccountInfo->UserId;
-		Params.LocalName = TEXT("test");
-		Params.SchemaName = TEXT("test");
-		Params.MaxMembers = 2;
-		Params.JoinPolicy = ELobbyJoinPolicy::InvitationOnly;
-		Params.Attributes.Add(TEXT("test_attribute_key"), TEXT("test_attribute_value"));
-		Params.LocalUsers.Emplace(FJoinLobbyLocalUserData{User1Info.AccountInfo->UserId, {{TEXT("test_attribute_key"), TEXT("test_attribute_value")}}});
-		return Params;
-	})
-	.Then(CaptureOperationStepResult<FCreateLobby>(*Op, CreateLobbyKeyName, &FLobbiesCommon::CreateLobby))
-	.Then([this](TOnlineAsyncOp<FFunctionalTestLobbies>& InAsyncOp)
-	{
-		const FFunctionalTestLoginUser::Result& User1Info = GetOpDataChecked<FFunctionalTestLoginUser::Result>(InAsyncOp, User1KeyName);
-		const FCreateLobby::Result& CreateResult = GetOpDataChecked<FCreateLobby::Result>(InAsyncOp, CreateLobbyKeyName);
-
-		// Check both lobby joined events were received and in the correct order.
-		const TSharedRef<FLobbyEventCapture>& EventCapture = GetOpDataChecked<TSharedRef<FLobbyEventCapture>>(InAsyncOp, LobbyEventCaptureKeyName);
-		if (EventCapture->GetTotalNotificationsReceived() != 2
-			|| EventCapture->LobbyJoined.Num() != 1 || EventCapture->LobbyMemberJoined.Num() != 1
-			|| EventCapture->LobbyJoined[0].GlobalIndex > EventCapture->LobbyMemberJoined[0].GlobalIndex)
-		{
-			InAsyncOp.SetError(Errors::Cancelled());
-		}
-
-		// Check lobby attributes are set to the expected values.
-		if (!CreateResult.Lobby->Attributes.OrderIndependentCompareEqual(
-			TMap<FLobbyAttributeId, FLobbyVariant>{{TEXT("test_attribute_key"), TEXT("test_attribute_value")}}))
-		{
-			InAsyncOp.SetError(Errors::Cancelled());
-		}
-
-		const TSharedRef<const FLobbyMember>* MemberData = CreateResult.Lobby->Members.Find(User1Info.AccountInfo->UserId);
-		if (CreateResult.Lobby->Members.Num() == 1 && MemberData != nullptr)
-		{
-			if (!(**MemberData).Attributes.OrderIndependentCompareEqual(
-				TMap<FLobbyAttributeId, FLobbyVariant>{{TEXT("test_attribute_key"), TEXT("test_attribute_value")}}))
-			{
-				InAsyncOp.SetError(Errors::Cancelled());
-			}
-		}
-		else
-		{
-			InAsyncOp.SetError(Errors::Cancelled());
-		}
-	})
-	.Then([this](TOnlineAsyncOp<FFunctionalTestLobbies>& InAsyncOp)
-	{
-		// Prepare for next test check.
-		GetOpDataChecked<TSharedRef<FLobbyEventCapture>>(InAsyncOp, LobbyEventCaptureKeyName)->Empty();
-
-		const FFunctionalTestLoginUser::Result& User1Info = GetOpDataChecked<FFunctionalTestLoginUser::Result>(InAsyncOp, User1KeyName);
-		const FCreateLobby::Result& CreateLobbyResult = GetOpDataChecked<FCreateLobby::Result>(InAsyncOp, CreateLobbyKeyName);
-
-		FModifyLobbyAttributes::Params Params;
-		Params.LocalUserId = User1Info.AccountInfo->UserId;
-		Params.LobbyId = CreateLobbyResult.Lobby->LobbyId;
-		Params.MutatedAttributes = {{TEXT("test_attribute_key"), TEXT("mutated_test_attribute_value")}};
-		return Params;
-	})
-	.Then(ConsumeOperationStepResult<FModifyLobbyAttributes>(*Op, &FLobbiesCommon::ModifyLobbyAttributes))
-	.Then([this](TOnlineAsyncOp<FFunctionalTestLobbies>& InAsyncOp)
-	{
-		// Check that modification event was received.
-		const TSharedRef<FLobbyEventCapture>& EventCapture = GetOpDataChecked<TSharedRef<FLobbyEventCapture>>(InAsyncOp, LobbyEventCaptureKeyName);
-		if (EventCapture->LobbyAttributesChanged.Num() == 1)
-		{
-			const FLobbyAttributesChanged& Notification = EventCapture->LobbyAttributesChanged[0].Notification;
-			if (!Notification.ChangedAttributes.Difference(TSet<FLobbyAttributeId>{TEXT("test_attribute_key")}).IsEmpty() ||
-				!Notification.Lobby->Attributes.OrderIndependentCompareEqual(TMap<FLobbyAttributeId, FLobbyVariant>{{TEXT("test_attribute_key"), TEXT("mutated_test_attribute_value")}}))
-			{
-				InAsyncOp.SetError(Errors::Cancelled());
-			}
-		}
-		else
-		{
-			InAsyncOp.SetError(Errors::Cancelled());
-		}
-	})
-	.Then([this](TOnlineAsyncOp<FFunctionalTestLobbies>& InAsyncOp)
-	{
-		// Prepare for next test check.
-		GetOpDataChecked<TSharedRef<FLobbyEventCapture>>(InAsyncOp, LobbyEventCaptureKeyName)->Empty();
-
-		const FFunctionalTestLoginUser::Result& User1Info = GetOpDataChecked<FFunctionalTestLoginUser::Result>(InAsyncOp, User1KeyName);
-		const FCreateLobby::Result& CreateLobbyResult = GetOpDataChecked<FCreateLobby::Result>(InAsyncOp, CreateLobbyKeyName);
-
-		FModifyLobbyMemberAttributes::Params Params;
-		Params.LocalUserId = User1Info.AccountInfo->UserId;
-		Params.LobbyId = CreateLobbyResult.Lobby->LobbyId;
-		Params.MutatedAttributes = {{TEXT("test_attribute_key"), TEXT("mutated_test_attribute_value")}};
-		return Params;
-	})
-	.Then(ConsumeOperationStepResult<FModifyLobbyMemberAttributes>(*Op, &FLobbiesCommon::ModifyLobbyMemberAttributes))
-	.Then([this](TOnlineAsyncOp<FFunctionalTestLobbies>& InAsyncOp)
-	{
-		const FFunctionalTestLoginUser::Result& User1Info = GetOpDataChecked<FFunctionalTestLoginUser::Result>(InAsyncOp, User1KeyName);
-
-		// Check that modification event was received.
-		const TSharedRef<FLobbyEventCapture>& EventCapture = GetOpDataChecked<TSharedRef<FLobbyEventCapture>>(InAsyncOp, LobbyEventCaptureKeyName);
-		if (EventCapture->LobbyMemberAttributesChanged.Num() == 1)
-		{
-			const FLobbyMemberAttributesChanged& Notification = EventCapture->LobbyMemberAttributesChanged[0].Notification;
-			if (const TSharedRef<const FLobbyMember>* MemberData = Notification.Lobby->Members.Find(User1Info.AccountInfo->UserId))
-			{
-				if (!Notification.ChangedAttributes.Difference(TSet<FLobbyAttributeId>{TEXT("test_attribute_key")}).IsEmpty() ||
-					!(**MemberData).Attributes.OrderIndependentCompareEqual(TMap<FLobbyAttributeId, FLobbyVariant>{{TEXT("test_attribute_key"), TEXT("mutated_test_attribute_value")}}))
-				{
-					InAsyncOp.SetError(Errors::Cancelled());
-				}
-			}
-			else
-			{
-				InAsyncOp.SetError(Errors::Cancelled());
-			}
-		}
-		else
-		{
-			InAsyncOp.SetError(Errors::Cancelled());
-		}
-	})
-	.Then([this](TOnlineAsyncOp<FFunctionalTestLobbies>& InAsyncOp)
-	{
-		// Prepare for next test check.
-		GetOpDataChecked<TSharedRef<FLobbyEventCapture>>(InAsyncOp, LobbyEventCaptureKeyName)->Empty();
-
-		const FFunctionalTestLoginUser::Result& User1Info = GetOpDataChecked<FFunctionalTestLoginUser::Result>(InAsyncOp, User1KeyName);
-		const FCreateLobby::Result& CreateLobbyResult = GetOpDataChecked<FCreateLobby::Result>(InAsyncOp, CreateLobbyKeyName);
-
-		FLeaveLobby::Params Params;
-		Params.LocalUserId = User1Info.AccountInfo->UserId;
-		Params.LobbyId = CreateLobbyResult.Lobby->LobbyId;
-		return Params;
-	})
-	.Then(ConsumeOperationStepResult<FLeaveLobby>(*Op, &FLobbiesCommon::LeaveLobby))
-	.Then([this](TOnlineAsyncOp<FFunctionalTestLobbies>& InAsyncOp)
-	{
-		// Check for expected events.
-		const TSharedRef<FLobbyEventCapture>& EventCapture = GetOpDataChecked<TSharedRef<FLobbyEventCapture>>(InAsyncOp, LobbyEventCaptureKeyName);
-		if (EventCapture->GetTotalNotificationsReceived() != 2
-			|| EventCapture->LobbyLeft.Num() != 1 || EventCapture->LobbyMemberLeft.Num() != 1
-			|| EventCapture->LobbyLeft[0].GlobalIndex < EventCapture->LobbyMemberLeft[0].GlobalIndex)
-		{
-			InAsyncOp.SetError(Errors::Cancelled());
-		}
-	})
-
-	//----------------------------------------------------------------------------------------------
-	// Test 2:
-	//    Step 1: Create a lobby with primary user.
-	//    Step 2: Primary user invites secondary user.
-	//    Step 3: Join lobby with secondary user.
-	//    Step 4: Leave lobby with secondary user.
-	//    Step 5: Leave lobby with primary user.
-	//----------------------------------------------------------------------------------------------
-
-	.Then([this](TOnlineAsyncOp<FFunctionalTestLobbies>& InAsyncOp)
-	{
-		// Prepare for next test check.
-		GetOpDataChecked<TSharedRef<FLobbyEventCapture>>(InAsyncOp, LobbyEventCaptureKeyName)->Empty();
-
-		const FFunctionalTestLoginUser::Result& User1Info = GetOpDataChecked<FFunctionalTestLoginUser::Result>(InAsyncOp, User1KeyName);
-
-		FCreateLobby::Params Params;
-		Params.LocalUserId = User1Info.AccountInfo->UserId;
-		Params.LocalName = TEXT("test");
-		Params.SchemaName = TEXT("test");
-		Params.MaxMembers = 2;
-		Params.JoinPolicy = ELobbyJoinPolicy::InvitationOnly;
-		//Params.Attributes;
-		Params.LocalUsers.Emplace(FJoinLobbyLocalUserData{User1Info.AccountInfo->UserId, {}});
-		return Params;
-	})
-	.Then(CaptureOperationStepResult<FCreateLobby>(*Op, CreateLobbyKeyName, &FLobbiesCommon::CreateLobby))
-	.Then([this](TOnlineAsyncOp<FFunctionalTestLobbies>& InAsyncOp)
-	{
-		// Check both lobby joined events were received and in the correct order.
-		const TSharedRef<FLobbyEventCapture>& EventCapture = GetOpDataChecked<TSharedRef<FLobbyEventCapture>>(InAsyncOp, LobbyEventCaptureKeyName);
-		if (EventCapture->GetTotalNotificationsReceived() != 2
-			|| EventCapture->LobbyJoined.Num() != 1 || EventCapture->LobbyMemberJoined.Num() != 1
-			|| EventCapture->LobbyJoined[0].GlobalIndex > EventCapture->LobbyMemberJoined[0].GlobalIndex)
-		{
-			InAsyncOp.SetError(Errors::Cancelled());
-		}
-	})
-	.Then([this](TOnlineAsyncOp<FFunctionalTestLobbies>& InAsyncOp)
-	{
-		// Prepare for next test check.
-		GetOpDataChecked<TSharedRef<FLobbyEventCapture>>(InAsyncOp, LobbyEventCaptureKeyName)->Empty();
-
-		const FFunctionalTestLoginUser::Result& User1Info = GetOpDataChecked<FFunctionalTestLoginUser::Result>(InAsyncOp, User1KeyName);
-		const FFunctionalTestLoginUser::Result& User2Info = GetOpDataChecked<FFunctionalTestLoginUser::Result>(InAsyncOp, User2KeyName);
-		const FCreateLobby::Result& CreateLobbyResult = GetOpDataChecked<FCreateLobby::Result>(InAsyncOp, CreateLobbyKeyName);
-
-		FInviteLobbyMember::Params Params;
-		Params.LocalUserId = User1Info.AccountInfo->UserId;
-		Params.LobbyId = CreateLobbyResult.Lobby->LobbyId;
-		Params.TargetUserId = User2Info.AccountInfo->UserId;
-		return Params;
-	})
-	.Then(ConsumeOperationStepResult<FInviteLobbyMember>(*Op, &FLobbiesCommon::InviteLobbyMember))
-	.Then([this, TestConfig](TOnlineAsyncOp<FFunctionalTestLobbies>& InAsyncOp)
-	{
-		const FFunctionalTestLoginUser::Result& User2Info = GetOpDataChecked<FFunctionalTestLoginUser::Result>(InAsyncOp, User2KeyName);
-		const FCreateLobby::Result& CreateLobbyResult = GetOpDataChecked<FCreateLobby::Result>(InAsyncOp, CreateLobbyKeyName);
-		return AwaitInvitation(User2Info.AccountInfo->UserId, CreateLobbyResult.Lobby->LobbyId, TestConfig->InvitationWaitSeconds);
-	})
-	.Then([this](TOnlineAsyncOp<FFunctionalTestLobbies>& InAsyncOp, TDefaultErrorResultInternal<FOnlineLobbyIdHandle>&& Result)
-	{
-		if (Result.IsError())
-		{
-			InAsyncOp.SetError(Errors::Cancelled(Result.GetErrorValue()));
-		}
-	})
-	.Then([this](TOnlineAsyncOp<FFunctionalTestLobbies>& InAsyncOp)
-	{
-		// Prepare for next test check.
-		GetOpDataChecked<TSharedRef<FLobbyEventCapture>>(InAsyncOp, LobbyEventCaptureKeyName)->Empty();
-
-		const FFunctionalTestLoginUser::Result& User2Info = GetOpDataChecked<FFunctionalTestLoginUser::Result>(InAsyncOp, User2KeyName);
-		const FCreateLobby::Result& CreateLobbyResult = GetOpDataChecked<FCreateLobby::Result>(InAsyncOp, CreateLobbyKeyName);
-
-		FJoinLobby::Params Params;
-		Params.LocalUserId = User2Info.AccountInfo->UserId;
-		Params.LocalName = TEXT("test");
-		Params.LobbyId = CreateLobbyResult.Lobby->LobbyId;
-		Params.LocalUsers.Emplace(FJoinLobbyLocalUserData{User2Info.AccountInfo->UserId, {}});
-		return Params;
-	})
-	.Then(ConsumeOperationStepResult<FJoinLobby>(*Op, &FLobbiesCommon::JoinLobby))
-	.Then([this](TOnlineAsyncOp<FFunctionalTestLobbies>& InAsyncOp)
-	{
-		// Check both lobby joined events were received and in the correct order.
-		const TSharedRef<FLobbyEventCapture>& EventCapture = GetOpDataChecked<TSharedRef<FLobbyEventCapture>>(InAsyncOp, LobbyEventCaptureKeyName);
-		if (EventCapture->GetTotalNotificationsReceived() != 1 || EventCapture->LobbyMemberJoined.Num() != 1)
-		{
-			InAsyncOp.SetError(Errors::Cancelled());
-		}
-	})
-	.Then([this](TOnlineAsyncOp<FFunctionalTestLobbies>& InAsyncOp)
-	{
-		// Prepare for next test check.
-		GetOpDataChecked<TSharedRef<FLobbyEventCapture>>(InAsyncOp, LobbyEventCaptureKeyName)->Empty();
-
-		const FFunctionalTestLoginUser::Result& User2Info = GetOpDataChecked<FFunctionalTestLoginUser::Result>(InAsyncOp, User2KeyName);
-		const FCreateLobby::Result& CreateLobbyResult = GetOpDataChecked<FCreateLobby::Result>(InAsyncOp, CreateLobbyKeyName);
-
-		FLeaveLobby::Params Params;
-		Params.LocalUserId = User2Info.AccountInfo->UserId;
-		Params.LobbyId = CreateLobbyResult.Lobby->LobbyId;
-		return Params;
-	})
-	.Then(ConsumeOperationStepResult<FLeaveLobby>(*Op, &FLobbiesCommon::LeaveLobby))
-	.Then([this](TOnlineAsyncOp<FFunctionalTestLobbies>& InAsyncOp)
-	{
-		// Check for expected events.
-		const TSharedRef<FLobbyEventCapture>& EventCapture = GetOpDataChecked<TSharedRef<FLobbyEventCapture>>(InAsyncOp, LobbyEventCaptureKeyName);
-		if (EventCapture->GetTotalNotificationsReceived() != 1 || EventCapture->LobbyMemberLeft.Num() != 1)
-		{
-			InAsyncOp.SetError(Errors::Cancelled());
-		}
-	})
-	.Then([this](TOnlineAsyncOp<FFunctionalTestLobbies>& InAsyncOp)
-	{
-		// Prepare for next test check.
-		GetOpDataChecked<TSharedRef<FLobbyEventCapture>>(InAsyncOp, LobbyEventCaptureKeyName)->Empty();
-
-		const FFunctionalTestLoginUser::Result& User1Info = GetOpDataChecked<FFunctionalTestLoginUser::Result>(InAsyncOp, User1KeyName);
-		const FCreateLobby::Result& CreateLobbyResult = GetOpDataChecked<FCreateLobby::Result>(InAsyncOp, CreateLobbyKeyName);
-
-		FLeaveLobby::Params Params;
-		Params.LocalUserId = User1Info.AccountInfo->UserId;
-		Params.LobbyId = CreateLobbyResult.Lobby->LobbyId;
-		return Params;
-	})
-	.Then(ConsumeOperationStepResult<FLeaveLobby>(*Op, &FLobbiesCommon::LeaveLobby))
-	.Then([this](TOnlineAsyncOp<FFunctionalTestLobbies>& InAsyncOp)
-	{
-		// Check for expected events.
-		const TSharedRef<FLobbyEventCapture>& EventCapture = GetOpDataChecked<TSharedRef<FLobbyEventCapture>>(InAsyncOp, LobbyEventCaptureKeyName);
-		if (EventCapture->GetTotalNotificationsReceived() != 2
-			|| EventCapture->LobbyLeft.Num() != 1 || EventCapture->LobbyMemberLeft.Num() != 1
-			|| EventCapture->LobbyLeft[0].GlobalIndex < EventCapture->LobbyMemberLeft[0].GlobalIndex)
-		{
-			InAsyncOp.SetError(Errors::Cancelled());
-		}
-	})
-
-	//----------------------------------------------------------------------------------------------
-	// Test 3: Simple search
-	//    Step 1: Create a lobby with primary user.
-	//    Step 2: Secondary user searches for lobbies.
-	//    Step 3: Join lobby with secondary user.
-	//    Step 4: Leave lobby with secondary user.
-	//    Step 5: Leave lobby with primary user.
-	//----------------------------------------------------------------------------------------------
-
-	.Then([this](TOnlineAsyncOp<FFunctionalTestLobbies>& InAsyncOp)
-	{
-		// Prepare for next test check.
-		GetOpDataChecked<TSharedRef<FLobbyEventCapture>>(InAsyncOp, LobbyEventCaptureKeyName)->Empty();
-
-		const FFunctionalTestLoginUser::Result& User1Info = GetOpDataChecked<FFunctionalTestLoginUser::Result>(InAsyncOp, User1KeyName);
-
-		FCreateLobby::Params Params;
-		Params.LocalUserId = User1Info.AccountInfo->UserId;
-		Params.LocalName = TEXT("test");
-		Params.SchemaName = TEXT("test");
-		Params.MaxMembers = 2;
-		Params.JoinPolicy = ELobbyJoinPolicy::PublicAdvertised;
-		//Params.Attributes;
-		Params.LocalUsers.Emplace(FJoinLobbyLocalUserData{User1Info.AccountInfo->UserId, {}});
-		return Params;
-	})
-	.Then(CaptureOperationStepResult<FCreateLobby>(*Op, CreateLobbyKeyName, &FLobbiesCommon::CreateLobby))
-	.Then([this](TOnlineAsyncOp<FFunctionalTestLobbies>& InAsyncOp)
-	{
-		// Check both lobby joined events were received and in the correct order.
-		const TSharedRef<FLobbyEventCapture>& EventCapture = GetOpDataChecked<TSharedRef<FLobbyEventCapture>>(InAsyncOp, LobbyEventCaptureKeyName);
-		if (EventCapture->GetTotalNotificationsReceived() != 2
-			|| EventCapture->LobbyJoined.Num() != 1 || EventCapture->LobbyMemberJoined.Num() != 1
-			|| EventCapture->LobbyJoined[0].GlobalIndex > EventCapture->LobbyMemberJoined[0].GlobalIndex)
-		{
-			InAsyncOp.SetError(Errors::Cancelled());
-		}
-	})
-	.Then([this](TOnlineAsyncOp<FFunctionalTestLobbies>& InAsyncOp)
-	{
-		const FFunctionalTestLoginUser::Result& User2Info = GetOpDataChecked<FFunctionalTestLoginUser::Result>(InAsyncOp, User2KeyName);
-		const FCreateLobby::Result& CreateLobbyResult = GetOpDataChecked<FCreateLobby::Result>(InAsyncOp, CreateLobbyKeyName);
-
-		// Search for lobby from create
-		FFindLobbies::Params Params;
-		Params.LocalUserId = User2Info.AccountInfo->UserId;
-		Params.LobbyId = CreateLobbyResult.Lobby->LobbyId;
-		return Params;
-	})
-	.Then(CaptureOperationStepResult<FFindLobbies>(*Op, FindLobbyKeyName, &FLobbiesCommon::FindLobbies))
-	.Then([this](TOnlineAsyncOp<FFunctionalTestLobbies>& InAsyncOp)
-	{
-		// Prepare for next test check.
-		GetOpDataChecked<TSharedRef<FLobbyEventCapture>>(InAsyncOp, LobbyEventCaptureKeyName)->Empty();
-
-		const FFunctionalTestLoginUser::Result& User2Info = GetOpDataChecked<FFunctionalTestLoginUser::Result>(InAsyncOp, User2KeyName);
-		const FCreateLobby::Result& CreateLobbyResult = GetOpDataChecked<FCreateLobby::Result>(InAsyncOp, CreateLobbyKeyName);
-		const FFindLobbies::Result& FindResults = GetOpDataChecked<FFindLobbies::Result>(InAsyncOp, FindLobbyKeyName);
-
-		if (FindResults.Lobbies.Num() != 1)
-		{
-			InAsyncOp.SetError(Errors::Cancelled());
-			return FJoinLobby::Params{};
-		}
-
-		FJoinLobby::Params Params;
-		Params.LocalUserId = User2Info.AccountInfo->UserId;
-		Params.LocalName = TEXT("test");
-		Params.LobbyId = FindResults.Lobbies[0]->LobbyId;
-		Params.LocalUsers.Emplace(FJoinLobbyLocalUserData{User2Info.AccountInfo->UserId, {}});
-		return Params;
-	})
-	.Then(ConsumeOperationStepResult<FJoinLobby>(*Op, &FLobbiesCommon::JoinLobby))
-	.Then([this](TOnlineAsyncOp<FFunctionalTestLobbies>& InAsyncOp)
-	{
-		// Check both lobby joined events were received and in the correct order.
-		const TSharedRef<FLobbyEventCapture>& EventCapture = GetOpDataChecked<TSharedRef<FLobbyEventCapture>>(InAsyncOp, LobbyEventCaptureKeyName);
-		if (EventCapture->GetTotalNotificationsReceived() != 1 || EventCapture->LobbyMemberJoined.Num() != 1)
-		{
-			InAsyncOp.SetError(Errors::Cancelled());
-		}
-	})
-	.Then([this](TOnlineAsyncOp<FFunctionalTestLobbies>& InAsyncOp)
-	{
-		// Prepare for next test check.
-		GetOpDataChecked<TSharedRef<FLobbyEventCapture>>(InAsyncOp, LobbyEventCaptureKeyName)->Empty();
-
-		const FFunctionalTestLoginUser::Result& User2Info = GetOpDataChecked<FFunctionalTestLoginUser::Result>(InAsyncOp, User2KeyName);
-		const FCreateLobby::Result& CreateLobbyResult = GetOpDataChecked<FCreateLobby::Result>(InAsyncOp, CreateLobbyKeyName);
-
-		FLeaveLobby::Params Params;
-		Params.LocalUserId = User2Info.AccountInfo->UserId;
-		Params.LobbyId = CreateLobbyResult.Lobby->LobbyId;
-		return Params;
-	})
-	.Then(ConsumeOperationStepResult<FLeaveLobby>(*Op, &FLobbiesCommon::LeaveLobby))
-	.Then([this](TOnlineAsyncOp<FFunctionalTestLobbies>& InAsyncOp)
-	{
-		// Check for expected events.
-		const TSharedRef<FLobbyEventCapture>& EventCapture = GetOpDataChecked<TSharedRef<FLobbyEventCapture>>(InAsyncOp, LobbyEventCaptureKeyName);
-		if (EventCapture->GetTotalNotificationsReceived() != 1|| EventCapture->LobbyMemberLeft.Num() != 1)
-		{
-			InAsyncOp.SetError(Errors::Cancelled());
-		}
-	})
-	.Then([this](TOnlineAsyncOp<FFunctionalTestLobbies>& InAsyncOp)
-	{
-		// Prepare for next test check.
-		GetOpDataChecked<TSharedRef<FLobbyEventCapture>>(InAsyncOp, LobbyEventCaptureKeyName)->Empty();
-
-		const FFunctionalTestLoginUser::Result& User1Info = GetOpDataChecked<FFunctionalTestLoginUser::Result>(InAsyncOp, User1KeyName);
-		const FCreateLobby::Result& CreateLobbyResult = GetOpDataChecked<FCreateLobby::Result>(InAsyncOp, CreateLobbyKeyName);
-
-		FLeaveLobby::Params Params;
-		Params.LocalUserId = User1Info.AccountInfo->UserId;
-		Params.LobbyId = CreateLobbyResult.Lobby->LobbyId;
-		return Params;
-	})
-	.Then(ConsumeOperationStepResult<FLeaveLobby>(*Op, &FLobbiesCommon::LeaveLobby))
-	.Then([this](TOnlineAsyncOp<FFunctionalTestLobbies>& InAsyncOp)
-	{
-		// Check for expected events.
-		const TSharedRef<FLobbyEventCapture>& EventCapture = GetOpDataChecked<TSharedRef<FLobbyEventCapture>>(InAsyncOp, LobbyEventCaptureKeyName);
-		if (EventCapture->GetTotalNotificationsReceived() != 2
-			|| EventCapture->LobbyLeft.Num() != 1 || EventCapture->LobbyMemberLeft.Num() != 1
-			|| EventCapture->LobbyLeft[0].GlobalIndex < EventCapture->LobbyMemberLeft[0].GlobalIndex)
-		{
-			InAsyncOp.SetError(Errors::Cancelled());
-		}
-	})
-
-	//----------------------------------------------------------------------------------------------
-	// Complete
-	//----------------------------------------------------------------------------------------------
-
-	.Then([this](TOnlineAsyncOp<FFunctionalTestLobbies>& InAsyncOp)
-	{
-		InAsyncOp.SetResult(FFunctionalTestLobbies::Result{});
-	})
-	.Enqueue(GetServices().GetParallelQueue());
-
-	return Op->GetHandle();
-}
-
-TFuture<TOnlineResult<FLobbiesCommon::FFunctionalTestLoginUser>> FLobbiesCommon::FunctionalTestLoginUser(FFunctionalTestLoginUser::Params&& Params)
-{
-	UE_LOG(LogTemp, Log, TEXT("[FLobbiesCommon::FunctionalTestLoginUser] Logging in user. PlatformUserId: %d"), FPlatformMisc::GetUserIndexForPlatformUser(Params.PlatformUserId));
-
 	IAuth* AuthInterface = Services.Get<IAuth>();
 	if (!AuthInterface)
 	{
-		UE_LOG(LogTemp, Warning, TEXT("[FLobbiesCommon::FunctionalTestLoginUser] Login failed. PlatformUserId: %d"), FPlatformMisc::GetUserIndexForPlatformUser(Params.PlatformUserId));
-		return MakeFulfilledPromise<TOnlineResult<FFunctionalTestLoginUser>>(Errors::MissingInterface()).GetFuture();
-	}
-
-	TSharedRef<TPromise<TOnlineResult<FFunctionalTestLoginUser>>> Promise = MakeShared<TPromise<TOnlineResult<FFunctionalTestLoginUser>>>();
-	TFuture<TOnlineResult<FFunctionalTestLoginUser>> Future = Promise->GetFuture();
-
-	FAuthLogin::Params LoginParams;
-	LoginParams.PlatformUserId = Params.PlatformUserId;
-	LoginParams.CredentialsType = Params.Type;
-	LoginParams.CredentialsId = Params.Id;
-	LoginParams.CredentialsToken.Set<FString>(Params.Token);
-
-	AuthInterface->Login(MoveTemp(LoginParams))
-	.OnComplete([Promise, PlatformUserId = Params.PlatformUserId](const TOnlineResult<FAuthLogin>& LoginResult)
-	{
-		if (LoginResult.IsError())
-		{
-			UE_LOG(LogTemp, Warning, TEXT("[FLobbiesCommon::FunctionalTestLoginUser] Login failed. PlatformUserId: %d, Result: %s"), FPlatformMisc::GetUserIndexForPlatformUser(PlatformUserId), *LoginResult.GetErrorValue().GetLogString());
-			Promise->EmplaceValue(Errors::RequestFailure());
-		}
-		else
-		{
-			UE_LOG(LogTemp, Log, TEXT("[FLobbiesCommon::FunctionalTestLoginUser] Login Succeeded. PlatformUserId: %d, User: %s"), FPlatformMisc::GetUserIndexForPlatformUser(PlatformUserId), *ToLogString(LoginResult.GetOkValue().AccountInfo->UserId));
-			Promise->EmplaceValue(FFunctionalTestLoginUser::Result{LoginResult.GetOkValue().AccountInfo});
-		}
-	});
-
-	return Future;
-}
-
-TFuture<TOnlineResult<FLobbiesCommon::FFunctionalTestLogoutUser>> FLobbiesCommon::FunctionalTestLogoutUser(FFunctionalTestLogoutUser::Params&& Params)
-{
-	UE_LOG(LogTemp, Log, TEXT("[FLobbiesCommon::FunctionalTestLogoutUser] Logging out user. PlatformUserId: %d"), FPlatformMisc::GetUserIndexForPlatformUser(Params.PlatformUserId));
-
-	IAuth* AuthInterface = Services.Get<IAuth>();
-	if (!AuthInterface)
-	{
-		UE_LOG(LogTemp, Warning, TEXT("[FLobbiesCommon::FunctionalTestLogoutUser] Logout failed. PlatformUserId: %d"), FPlatformMisc::GetUserIndexForPlatformUser(Params.PlatformUserId));
-		return MakeFulfilledPromise<TOnlineResult<FFunctionalTestLogoutUser>>(Errors::MissingInterface()).GetFuture();
-	}
-
-	FAuthGetAccountByPlatformUserId::Params GetAccountParams;
-	GetAccountParams.PlatformUserId = Params.PlatformUserId;
-	TOnlineResult<FAuthGetAccountByPlatformUserId> Result = AuthInterface->GetAccountByPlatformUserId(MoveTemp(GetAccountParams));
-	if (Result.IsError())
-	{
-		// Ignore errors for now. This function should ignore logged out users.
-		if (Result.GetErrorValue() == Errors::Unknown())
-		{
-			UE_LOG(LogTemp, Log, TEXT("[FLobbiesCommon::FunctionalTestLogoutUser] Logout Succeeded - already logged out. PlatformUserId: %d"), FPlatformMisc::GetUserIndexForPlatformUser(Params.PlatformUserId));
-			return MakeFulfilledPromise<TOnlineResult<FFunctionalTestLogoutUser>>(FFunctionalTestLogoutUser::Result{}).GetFuture();
-		}
-		else
-		{
-			UE_LOG(LogTemp, Warning, TEXT("[FLobbiesCommon::FunctionalTestLogoutUser] Logout failed. PlatformUserId: %d, Result: %s"), FPlatformMisc::GetUserIndexForPlatformUser(Params.PlatformUserId), *Result.GetErrorValue().GetLogString());
-			return MakeFulfilledPromise<TOnlineResult<FFunctionalTestLogoutUser>>(Errors::Unknown(MoveTemp(Result.GetErrorValue()))).GetFuture();
-		}
+		TOnlineAsyncOpRef<FFunctionalTestLobbies> Operation = GetOp<FFunctionalTestLobbies>(MoveTemp(Params));
+		Operation->SetError(Errors::MissingInterface());
+		return Operation->GetHandle();
 	}
 
-	TSharedRef<FAccountInfo> AccountInfo = Result.GetOkValue().AccountInfo;
-	if (AccountInfo->LoginStatus == ELoginStatus::NotLoggedIn)
-	{
-		return MakeFulfilledPromise<TOnlineResult<FFunctionalTestLogoutUser>>(Errors::InvalidUser()).GetFuture();
-	}
-
-	TSharedRef<TPromise<TOnlineResult<FFunctionalTestLogoutUser>>> Promise = MakeShared<TPromise<TOnlineResult<FFunctionalTestLogoutUser>>>();
-	TFuture<TOnlineResult<FFunctionalTestLogoutUser>> Future = Promise->GetFuture();
-
-	FAuthLogout::Params LogoutParams;
-	LogoutParams.LocalUserId = AccountInfo->UserId;
-
-	AuthInterface->Logout(MoveTemp(LogoutParams))
-	.OnComplete([Promise, PlatformUserId = Params.PlatformUserId](const TOnlineResult<FAuthLogout>& LogoutResult)
-	{
-		if (LogoutResult.IsError())
-		{
-			UE_LOG(LogTemp, Warning, TEXT("[FLobbiesCommon::FunctionalTestLogoutUser] Logout failed. PlatformUserId: %d, Result: %s"), FPlatformMisc::GetUserIndexForPlatformUser(PlatformUserId), *LogoutResult.GetErrorValue().GetLogString());
-			Promise->EmplaceValue(Errors::RequestFailure());
-		}
-		else
-		{
-			UE_LOG(LogTemp, Log, TEXT("[FLobbiesCommon::FunctionalTestLogoutUser] Logout Succeeded. PlatformUserId: %d"), FPlatformMisc::GetUserIndexForPlatformUser(PlatformUserId));
-			Promise->EmplaceValue(FFunctionalTestLogoutUser::Result{});
-		}
-	});
-
-	return Future;
+	return RunLobbyFunctionalTest(*AuthInterface, *this, LobbyEvents);
 }
-
-TFuture<TOnlineResult<FLobbiesCommon::FFunctionalTestLogoutAllUsers>> FLobbiesCommon::FunctionalTestLogoutAllUsers(FFunctionalTestLogoutAllUsers::Params&& Params)
-{
-	UE_LOG(LogTemp, Log, TEXT("[FLobbiesCommon::FFunctionalTestLogoutAllUsers] Logging out all users."));
-
-	TArray<TFuture<TOnlineResult<FLobbiesCommon::FFunctionalTestLogoutUser>>> LogoutFutures;
-	for (int32 index = 0; index < MAX_LOCAL_PLAYERS; ++index)
-	{
-		LogoutFutures.Emplace(FunctionalTestLogoutUser(FFunctionalTestLogoutUser::Params{FPlatformMisc::GetPlatformUserForUserIndex(index)}));
-	}
-
-	TSharedRef<TPromise<TOnlineResult<FFunctionalTestLogoutAllUsers>>> Promise = MakeShared<TPromise<TOnlineResult<FFunctionalTestLogoutAllUsers>>>();
-	TFuture<TOnlineResult<FFunctionalTestLogoutAllUsers>> Future = Promise->GetFuture();
-
-	WhenAll(MoveTempIfPossible(LogoutFutures))
-	.Then([Promise = MoveTemp(Promise)](TFuture<TArray<TOnlineResult<FLobbiesCommon::FFunctionalTestLogoutUser>>>&& Results)
-	{
-		bool HasAnyError = false;
-		for (TOnlineResult<FLobbiesCommon::FFunctionalTestLogoutUser>& Result : Results.Get())
-		{
-			HasAnyError |= Result.IsError();
-		}
-
-		if (HasAnyError)
-		{
-			UE_LOG(LogTemp, Warning, TEXT("[FLobbiesCommon::FunctionalTestLogoutUser] Logout all users Failed."));
-			Promise->EmplaceValue(Errors::RequestFailure());
-		}
-		else
-		{
-			UE_LOG(LogTemp, Log, TEXT("[FLobbiesCommon::FFunctionalTestLogoutAllUsers] Logout all users Succeeded."));
-			Promise->EmplaceValue(FFunctionalTestLogoutAllUsers::Result{});
-		}
-	});
-
-	return Future;
-}
-
-#endif
+#endif // LOBBIES_FUNCTIONAL_TEST_ENABLED
 
 /* UE::Online */ }
diff --git a/Engine/Plugins/Online/Experimental/OnlineServices/Source/OnlineServicesCommon/Private/Online/LobbiesCommonTests.cpp b/Engine/Plugins/Online/Experimental/OnlineServices/Source/OnlineServicesCommon/Private/Online/LobbiesCommonTests.cpp
new file mode 100644
--- /dev/null
+++ b/Engine/Plugins/Online/Experimental/OnlineServices/Source/OnlineServicesCommon/Private/Online/LobbiesCommonTests.cpp
@@ -0,0 +1,1899 @@
+// Copyright Epic Games, Inc. All Rights Reserved.
+
+#include "LobbiesCommonTests.h"
+
+#include "Containers/Ticker.h"
+#include "Misc/AutomationTest.h"
+#include "Online/LobbiesCommon.h"
+#include "Online/Auth.h"
+#include "Online/OnlineAsyncOp.h"
+
+namespace UE::Online {
+
+class FLobbyEventCapture final
+{
+public:
+	UE_NONCOPYABLE(FLobbyEventCapture);
+
+	FLobbyEventCapture(FLobbyEvents& LobbyEvents)
+		: LobbyJoinedHandle(LobbyEvents.OnLobbyJoined.Add(this, &FLobbyEventCapture::OnLobbyJoined))
+		, LobbyLeftHandle(LobbyEvents.OnLobbyLeft.Add(this, &FLobbyEventCapture::OnLobbyLeft))
+		, LobbyMemberJoinedHandle(LobbyEvents.OnLobbyMemberJoined.Add(this, &FLobbyEventCapture::OnLobbyMemberJoined))
+		, LobbyMemberLeftHandle(LobbyEvents.OnLobbyMemberLeft.Add(this, &FLobbyEventCapture::OnLobbyMemberLeft))
+		, LobbyLeaderChangedHandle(LobbyEvents.OnLobbyLeaderChanged.Add(this, &FLobbyEventCapture::OnLobbyLeaderChanged))
+		, LobbySchemaChangedHandle(LobbyEvents.OnLobbySchemaChanged.Add(this, &FLobbyEventCapture::OnLobbySchemaChanged))
+		, LobbyAttributesChangedHandle(LobbyEvents.OnLobbyAttributesChanged.Add(this, &FLobbyEventCapture::OnLobbyAttributesChanged))
+		, LobbyMemberAttributesChangedHandle(LobbyEvents.OnLobbyMemberAttributesChanged.Add(this, &FLobbyEventCapture::OnLobbyMemberAttributesChanged))
+		, LobbyInvitationAddedHandle(LobbyEvents.OnLobbyInvitationAdded.Add(this, &FLobbyEventCapture::OnLobbyInvitationAdded))
+		, LobbyInvitationRemovedHandle(LobbyEvents.OnLobbyInvitationRemoved.Add(this, &FLobbyEventCapture::OnLobbyInvitationRemoved))
+	{
+	}
+
+	void Empty()
+	{
+		LobbyJoined.Empty();
+		LobbyLeft.Empty();
+		LobbyMemberJoined.Empty();
+		LobbyMemberLeft.Empty();
+		LobbyLeaderChanged.Empty();
+		LobbySchemaChanged.Empty();
+		LobbyAttributesChanged.Empty();
+		LobbyMemberAttributesChanged.Empty();
+		LobbyInvitationAdded.Empty();
+		LobbyInvitationRemoved.Empty();
+		NextIndex = 0;
+		TotalNotificationsReceived = 0;
+	}
+
+	uint32 GetTotalNotificationsReceived() const
+	{
+		return TotalNotificationsReceived;
+	}
+
+	template <typename NotificationType>
+	struct TNotificationInfo
+	{
+		NotificationType Notification;
+		int32 GlobalIndex = 0;
+	};
+
+	TArray<TNotificationInfo<FLobbyJoined>> LobbyJoined;
+	TArray<TNotificationInfo<FLobbyLeft>> LobbyLeft;
+	TArray<TNotificationInfo<FLobbyMemberJoined>> LobbyMemberJoined;
+	TArray<TNotificationInfo<FLobbyMemberLeft>> LobbyMemberLeft;
+	TArray<TNotificationInfo<FLobbyLeaderChanged>> LobbyLeaderChanged;
+	TArray<TNotificationInfo<FLobbySchemaChanged>> LobbySchemaChanged;
+	TArray<TNotificationInfo<FLobbyAttributesChanged>> LobbyAttributesChanged;
+	TArray<TNotificationInfo<FLobbyMemberAttributesChanged>> LobbyMemberAttributesChanged;
+	TArray<TNotificationInfo<FLobbyInvitationAdded>> LobbyInvitationAdded;
+	TArray<TNotificationInfo<FLobbyInvitationRemoved>> LobbyInvitationRemoved;
+
+private:
+	template <typename ContainerType, typename NotificationType>
+	void AddEvent(ContainerType& Container, const NotificationType& Notification)
+	{
+		Container.Add({Notification, NextIndex++});
+		++TotalNotificationsReceived;
+	}
+
+	void OnLobbyJoined(const FLobbyJoined& Notification) { AddEvent(LobbyJoined, Notification); }
+	void OnLobbyLeft(const FLobbyLeft& Notification) { AddEvent(LobbyLeft, Notification); }
+	void OnLobbyMemberJoined(const FLobbyMemberJoined& Notification) { AddEvent(LobbyMemberJoined, Notification); }
+	void OnLobbyMemberLeft(const FLobbyMemberLeft& Notification) { AddEvent(LobbyMemberLeft, Notification); }
+	void OnLobbyLeaderChanged(const FLobbyLeaderChanged& Notification) { AddEvent(LobbyLeaderChanged, Notification); }
+	void OnLobbySchemaChanged(const FLobbySchemaChanged& Notification) { AddEvent(LobbySchemaChanged, Notification); }
+	void OnLobbyAttributesChanged(const FLobbyAttributesChanged& Notification) { AddEvent(LobbyAttributesChanged, Notification); }
+	void OnLobbyMemberAttributesChanged(const FLobbyMemberAttributesChanged& Notification) { AddEvent(LobbyMemberAttributesChanged, Notification); }
+	void OnLobbyInvitationAdded(const FLobbyInvitationAdded& Notification) { AddEvent(LobbyInvitationAdded, Notification); }
+	void OnLobbyInvitationRemoved(const FLobbyInvitationRemoved& Notification) { AddEvent(LobbyInvitationRemoved, Notification); }
+
+	FOnlineEventDelegateHandle LobbyJoinedHandle;
+	FOnlineEventDelegateHandle LobbyLeftHandle;
+	FOnlineEventDelegateHandle LobbyMemberJoinedHandle;
+	FOnlineEventDelegateHandle LobbyMemberLeftHandle;
+	FOnlineEventDelegateHandle LobbyLeaderChangedHandle;
+	FOnlineEventDelegateHandle LobbySchemaChangedHandle;
+	FOnlineEventDelegateHandle LobbyAttributesChangedHandle;
+	FOnlineEventDelegateHandle LobbyMemberAttributesChangedHandle;
+	FOnlineEventDelegateHandle LobbyInvitationAddedHandle;
+	FOnlineEventDelegateHandle LobbyInvitationRemovedHandle;
+	int32 NextIndex = 0;
+	uint32 TotalNotificationsReceived = 0;
+};
+
+//--------------------------------------------------------------------------------------------------
+// Unit testing
+//--------------------------------------------------------------------------------------------------
+
+IMPLEMENT_SIMPLE_AUTOMATION_TEST(FClientLobbyDataTest, 
+	"System.Engine.Online.ClientLobbyDataTest",
+	EAutomationTestFlags::EditorContext | EAutomationTestFlags::EngineFilter)
+bool FClientLobbyDataTest::RunTest(const FString& Parameters)
+{
+	// 1. Applying snapshot with no local members does not fire events.
+	{
+		FLobbyEvents LobbyEvents;
+		FLobbyEventCapture EventCapture(LobbyEvents);
+		FOnlineLobbyIdHandle LobbyId(EOnlineServices::Null, 1);
+		FClientLobbyData ClientData(LobbyId);
+
+		FOnlineAccountIdHandle User1(EOnlineServices::Null, 1);
+		FOnlineAccountIdHandle User2(EOnlineServices::Null, 2);
+		FName LobbySchemaName = TEXT("SchemaName");
+		int32 LobbyMaxMembers = 5;
+		ELobbyJoinPolicy LobbyJoinPolicy = ELobbyJoinPolicy::PublicAdvertised;
+		FLobbyAttributeId LobbyAttribute1Key = TEXT("LobbyAttribute1");
+		FLobbyVariant LobbyAttribute1Value;
+		LobbyAttribute1Value.Set(TEXT("Attribute1"));
+		FLobbyAttributeId LobbyAttribute2Key = TEXT("LobbyAttribute2");
+		FLobbyVariant LobbyAttribute2Value;
+		LobbyAttribute2Value.Set(TEXT("Attribute2"));
+		FLobbyAttributeId LobbyMemberAttribute1Key = TEXT("LobbyMemberAttribute1");
+		FLobbyVariant LobbyMemberAttribute1Value;
+		LobbyMemberAttribute1Value.Set(TEXT("MemberAttribute1"));
+		FLobbyAttributeId LobbyMemberAttribute2Key = TEXT("LobbyMemberAttribute2");
+		FLobbyVariant LobbyMemberAttribute2Value;
+		LobbyMemberAttribute2Value.Set(TEXT("MemberAttribute2"));
+
+		// 1.1 setup.
+		{
+			TMap<FOnlineAccountIdHandle, TSharedRef<FClientLobbyMemberSnapshot>> LobbyMemberSnapshots;
+			{
+				TSharedRef<FClientLobbyMemberSnapshot> MemberSnapshot = MakeShared<FClientLobbyMemberSnapshot>();
+				MemberSnapshot->AccountId = User1;
+				MemberSnapshot->Attributes.Add(LobbyMemberAttribute1Key, LobbyMemberAttribute1Value);
+				LobbyMemberSnapshots.Add(User1, MemberSnapshot);
+			}
+			{
+				TSharedRef<FClientLobbyMemberSnapshot> MemberSnapshot = MakeShared<FClientLobbyMemberSnapshot>();
+				MemberSnapshot->AccountId = User2;
+				LobbyMemberSnapshots.Add(User2, MemberSnapshot);
+			}
+
+			FClientLobbySnapshot LobbySnapshot;
+			LobbySnapshot.OwnerAccountId = User1;
+			LobbySnapshot.SchemaName = LobbySchemaName;
+			LobbySnapshot.MaxMembers = LobbyMaxMembers;
+			LobbySnapshot.JoinPolicy = LobbyJoinPolicy;
+			LobbySnapshot.Attributes.Add(LobbyAttribute1Key, LobbyAttribute1Value);
+			LobbySnapshot.Members = {User1, User2};
+
+			TMap<FOnlineAccountIdHandle, ELobbyMemberLeaveReason> LeaveReasons;
+			FApplyLobbyUpdateResult Result = ClientData.ApplyLobbyUpdateFromServiceSnapshot(MoveTemp(LobbySnapshot), MoveTemp(LobbyMemberSnapshots), MoveTemp(LeaveReasons), &LobbyEvents);
+			UTEST_EQUAL("Check no local members left.", Result.LeavingLocalMembers.Num(), 0);
+		}
+
+		// 1.1.1. Initial snapshot does not trigger any events.
+		UTEST_EQUAL("Snapshot does not trigger any events.", EventCapture.GetTotalNotificationsReceived(), 0);
+		UTEST_EQUAL("Verify lobby owner.", ClientData.GetPublicData().OwnerAccountId, User1);
+		UTEST_EQUAL("Verify lobby schema.", ClientData.GetPublicData().SchemaName, LobbySchemaName);
+		UTEST_EQUAL("Verify lobby max members.", ClientData.GetPublicData().MaxMembers, LobbyMaxMembers);
+		UTEST_EQUAL("Verify lobby join policy.", ClientData.GetPublicData().JoinPolicy, LobbyJoinPolicy);
+		UTEST_EQUAL("Verify lobby attributes.", ClientData.GetPublicData().Attributes.Num(), 1);
+		UTEST_NOT_NULL("Verify lobby attributes.", ClientData.GetPublicData().Attributes.Find(LobbyAttribute1Key));
+		UTEST_EQUAL("Verify lobby attributes.", ClientData.GetPublicData().Attributes[LobbyAttribute1Key], LobbyAttribute1Value);
+		UTEST_EQUAL("Verify lobby members.", ClientData.GetPublicData().Members.Num(), 2);
+		UTEST_NOT_NULL("Verify lobby members.", ClientData.GetPublicData().Members.Find(User1));
+		UTEST_EQUAL("Verify lobby members.", ClientData.GetPublicData().Members[User1]->AccountId, User1);
+		UTEST_NOT_NULL("Verify lobby members.", ClientData.GetPublicData().Members.Find(User2));
+		UTEST_EQUAL("Verify lobby members.", ClientData.GetPublicData().Members[User2]->AccountId, User2);
+		UTEST_EQUAL("Verify lobby member attributes.", ClientData.GetPublicData().Members[User1]->Attributes.Num(), 1);
+		UTEST_NOT_NULL("Verify lobby member attributes.", ClientData.GetPublicData().Members[User1]->Attributes.Find(LobbyMemberAttribute1Key));
+		UTEST_EQUAL("Verify lobby member attributes.", ClientData.GetPublicData().Members[User1]->Attributes[LobbyMemberAttribute1Key], LobbyMemberAttribute1Value);
+
+		// 1.2 setup.
+		{
+			LobbyAttribute1Value.Set(TEXT("ModifiedAttribute1"));
+			LobbyMemberAttribute1Value.Set(TEXT("ModifiedMemberAttribute1"));
+			LobbySchemaName = TEXT("ModifiedSchemaName");
+
+			TMap<FOnlineAccountIdHandle, TSharedRef<FClientLobbyMemberSnapshot>> LobbyMemberSnapshots;
+			{
+				TSharedRef<FClientLobbyMemberSnapshot> MemberSnapshot = MakeShared<FClientLobbyMemberSnapshot>();
+				MemberSnapshot->AccountId = User1;
+				MemberSnapshot->Attributes.Add(LobbyMemberAttribute1Key, LobbyMemberAttribute1Value);
+				LobbyMemberSnapshots.Add(User1, MemberSnapshot);
+			}
+			{
+				TSharedRef<FClientLobbyMemberSnapshot> MemberSnapshot = MakeShared<FClientLobbyMemberSnapshot>();
+				MemberSnapshot->AccountId = User2;
+				LobbyMemberSnapshots.Add(User2, MemberSnapshot);
+			}
+
+			FClientLobbySnapshot LobbySnapshot;
+			LobbySnapshot.OwnerAccountId = User2;
+			LobbySnapshot.SchemaName = LobbySchemaName;
+			LobbySnapshot.MaxMembers = LobbyMaxMembers;
+			LobbySnapshot.JoinPolicy = LobbyJoinPolicy;
+			LobbySnapshot.Attributes.Add(LobbyAttribute1Key, LobbyAttribute1Value);
+			LobbySnapshot.Members = {User1, User2};
+
+			TMap<FOnlineAccountIdHandle, ELobbyMemberLeaveReason> LeaveReasons;
+			FApplyLobbyUpdateResult Result = ClientData.ApplyLobbyUpdateFromServiceSnapshot(MoveTemp(LobbySnapshot), MoveTemp(LobbyMemberSnapshots), MoveTemp(LeaveReasons), &LobbyEvents);
+			UTEST_EQUAL("Check no local members left.", Result.LeavingLocalMembers.Num(), 0);
+		}
+
+		// 1.2.1. Snapshot with changed lobby attribute does not trigger lobby attribute changed.
+		// 1.2.2. Snapshot with changed lobby member attribute does not trigger lobby member attribute changed.
+		// 1.2.3. Snapshot with changed lobby schema does not trigger lobby schema changed.
+		// 1.2.4. Snapshot with changed lobby leader does not trigger lobby leader changed.
+		UTEST_EQUAL("Snapshot does not trigger any events.", EventCapture.GetTotalNotificationsReceived(), 0);
+		UTEST_EQUAL("Verify lobby owner.", ClientData.GetPublicData().OwnerAccountId, User2);
+		UTEST_EQUAL("Verify lobby schema.", ClientData.GetPublicData().SchemaName, LobbySchemaName);
+		UTEST_EQUAL("Verify lobby max members.", ClientData.GetPublicData().MaxMembers, LobbyMaxMembers);
+		UTEST_EQUAL("Verify lobby join policy.", ClientData.GetPublicData().JoinPolicy, LobbyJoinPolicy);
+		UTEST_EQUAL("Verify lobby attributes.", ClientData.GetPublicData().Attributes.Num(), 1);
+		UTEST_NOT_NULL("Verify lobby attributes.", ClientData.GetPublicData().Attributes.Find(LobbyAttribute1Key));
+		UTEST_EQUAL("Verify lobby attributes.", ClientData.GetPublicData().Attributes[LobbyAttribute1Key], LobbyAttribute1Value);
+		UTEST_EQUAL("Verify lobby members.", ClientData.GetPublicData().Members.Num(), 2);
+		UTEST_NOT_NULL("Verify lobby members.", ClientData.GetPublicData().Members.Find(User1));
+		UTEST_NOT_NULL("Verify lobby members.", ClientData.GetPublicData().Members.Find(User2));
+		UTEST_EQUAL("Verify lobby member attributes.", ClientData.GetPublicData().Members[User1]->Attributes.Num(), 1);
+		UTEST_NOT_NULL("Verify lobby member attributes.", ClientData.GetPublicData().Members[User1]->Attributes.Find(LobbyMemberAttribute1Key));
+		UTEST_EQUAL("Verify lobby member attributes.", ClientData.GetPublicData().Members[User1]->Attributes[LobbyMemberAttribute1Key], LobbyMemberAttribute1Value);
+
+		// 1.3 setup.
+		{
+			TMap<FOnlineAccountIdHandle, TSharedRef<FClientLobbyMemberSnapshot>> LobbyMemberSnapshots;
+			{
+				TSharedRef<FClientLobbyMemberSnapshot> MemberSnapshot = MakeShared<FClientLobbyMemberSnapshot>();
+				MemberSnapshot->AccountId = User1;
+				LobbyMemberSnapshots.Add(User1, MemberSnapshot);
+			}
+
+			FClientLobbySnapshot LobbySnapshot;
+			LobbySnapshot.OwnerAccountId = User1;
+			LobbySnapshot.SchemaName = LobbySchemaName;
+			LobbySnapshot.MaxMembers = LobbyMaxMembers;
+			LobbySnapshot.JoinPolicy = LobbyJoinPolicy;
+			LobbySnapshot.Members = {User1};
+
+			TMap<FOnlineAccountIdHandle, ELobbyMemberLeaveReason> LeaveReasons;
+			FApplyLobbyUpdateResult Result = ClientData.ApplyLobbyUpdateFromServiceSnapshot(MoveTemp(LobbySnapshot), MoveTemp(LobbyMemberSnapshots), MoveTemp(LeaveReasons), &LobbyEvents);
+			UTEST_EQUAL("Check no local members left.", Result.LeavingLocalMembers.Num(), 0);
+		}
+
+		// 1.3.1. Snapshot with cleared lobby attribute does not trigger lobby attribute changed.
+		// 1.3.2. Snapshot with cleared lobby member attribute does not trigger lobby member attribute changed.
+		// 1.3.3. Snapshot with members leaving does not trigger lobby member left or lobby left.
+		UTEST_EQUAL("Snapshot does not trigger any events.", EventCapture.GetTotalNotificationsReceived(), 0);
+		UTEST_EQUAL("Verify lobby attributes.", ClientData.GetPublicData().Attributes.Num(), 0);
+		UTEST_EQUAL("Verify lobby members.", ClientData.GetPublicData().Members.Num(), 1);
+		UTEST_NOT_NULL("Verify lobby members.", ClientData.GetPublicData().Members.Find(User1));
+		UTEST_EQUAL("Verify lobby member attributes.", ClientData.GetPublicData().Members[User1]->Attributes.Num(), 0);
+	}
+
+	// 2. Applying snapshot with local members present does fire events.
+	{
+		FLobbyEvents LobbyEvents;
+		FLobbyEventCapture EventCapture(LobbyEvents);
+		FOnlineLobbyIdHandle LobbyId(EOnlineServices::Null, 1);
+		FClientLobbyData ClientData(LobbyId);
+
+		FOnlineAccountIdHandle User1(EOnlineServices::Null, 1);
+		FOnlineAccountIdHandle User2(EOnlineServices::Null, 2);
+		FName LobbySchemaName = TEXT("SchemaName");
+		int32 LobbyMaxMembers = 5;
+		ELobbyJoinPolicy LobbyJoinPolicy = ELobbyJoinPolicy::PublicAdvertised;
+		FLobbyAttributeId LobbyAttribute1Key = TEXT("LobbyAttribute1");
+		FLobbyVariant LobbyAttribute1Value;
+		LobbyAttribute1Value.Set(TEXT("Attribute1"));
+		FLobbyAttributeId LobbyAttribute2Key = TEXT("LobbyAttribute2");
+		FLobbyVariant LobbyAttribute2Value;
+		LobbyAttribute2Value.Set(TEXT("Attribute2"));
+		FLobbyAttributeId LobbyMemberAttribute1Key = TEXT("LobbyMemberAttribute1");
+		FLobbyVariant LobbyMemberAttribute1Value;
+		LobbyMemberAttribute1Value.Set(TEXT("MemberAttribute1"));
+		FLobbyAttributeId LobbyMemberAttribute2Key = TEXT("LobbyMemberAttribute2");
+		FLobbyVariant LobbyMemberAttribute2Value;
+		LobbyMemberAttribute2Value.Set(TEXT("MemberAttribute2"));
+
+		// 2.1 setup
+		{
+			TMap<FOnlineAccountIdHandle, TSharedRef<FClientLobbyMemberSnapshot>> LobbyMemberSnapshots;
+			{
+				TSharedRef<FClientLobbyMemberSnapshot> MemberSnapshot = MakeShared<FClientLobbyMemberSnapshot>();
+				MemberSnapshot->AccountId = User1;
+				LobbyMemberSnapshots.Add(User1, MemberSnapshot);
+			}
+			{
+				TSharedRef<FClientLobbyMemberSnapshot> MemberSnapshot = MakeShared<FClientLobbyMemberSnapshot>();
+				MemberSnapshot->AccountId = User2;
+				MemberSnapshot->Attributes.Add(LobbyMemberAttribute1Key, LobbyMemberAttribute1Value);
+				LobbyMemberSnapshots.Add(User2, MemberSnapshot);
+			}
+
+			FClientLobbySnapshot LobbySnapshot;
+			LobbySnapshot.OwnerAccountId = User1;
+			LobbySnapshot.SchemaName = LobbySchemaName;
+			LobbySnapshot.MaxMembers = LobbyMaxMembers;
+			LobbySnapshot.JoinPolicy = LobbyJoinPolicy;
+			LobbySnapshot.Attributes.Add(LobbyAttribute1Key, LobbyAttribute1Value);
+			LobbySnapshot.Members = {User1, User2};
+
+			TMap<FOnlineAccountIdHandle, ELobbyMemberLeaveReason> LeaveReasons;
+			FApplyLobbyUpdateResult Result = ClientData.ApplyLobbyUpdateFromServiceSnapshot(MoveTemp(LobbySnapshot), MoveTemp(LobbyMemberSnapshots), MoveTemp(LeaveReasons), &LobbyEvents);
+			UTEST_EQUAL("Check no members left.", Result.LeavingLocalMembers.Num(), 0);
+		}
+
+		// 2.1.1. Initial snapshot does not trigger lobby joined or lobby member joined. - Member in snapshot is local member to simulate lobby creation.
+		UTEST_EQUAL("Snapshot does not trigger any events.", EventCapture.GetTotalNotificationsReceived(), 0);
+		UTEST_EQUAL("Verify lobby members.", ClientData.GetPublicData().Members.Num(), 2);
+		UTEST_NOT_NULL("Verify lobby members.", ClientData.GetPublicData().Members.Find(User1));
+		UTEST_NOT_NULL("Verify lobby members.", ClientData.GetPublicData().Members.Find(User2));
+
+		// 2.2 setup
+		{
+			EventCapture.Empty();
+
+			FClientLobbyDataChanges LobbyChanges;
+			LobbyChanges.MutatedMembers.Add(User1, MakeShared<FClientLobbyMemberDataChanges>());
+			FApplyLobbyUpdateResult Result = ClientData.ApplyLobbyUpdateFromLocalChanges(MoveTemp(LobbyChanges), LobbyEvents);
+			UTEST_EQUAL("Check no local members left.", Result.LeavingLocalMembers.Num(), 0);
+		}
+
+		// 2.2.1. Local member added to lobby through local change. Adding a local member will cause events to begin triggering.
+		UTEST_EQUAL("Check lobby and member joined events were received.", EventCapture.GetTotalNotificationsReceived(), 3);
+		UTEST_EQUAL("Check lobby joined event received.", EventCapture.LobbyJoined.Num(), 1);
+		UTEST_EQUAL("Check lobby joined event received first.", EventCapture.LobbyJoined[0].GlobalIndex, 0);
+		UTEST_EQUAL("Check lobby joined event is valid.", EventCapture.LobbyJoined[0].Notification.Lobby, ClientData.GetPublicDataPtr());
+		UTEST_EQUAL("Check member joined events received.", EventCapture.LobbyMemberJoined.Num(), 2);
+		UTEST_EQUAL("Check member joined event is valid.", EventCapture.LobbyMemberJoined[0].Notification.Lobby, ClientData.GetPublicDataPtr());
+		UTEST_TRUE("Check member joined event is valid.", EventCapture.LobbyMemberJoined[0].Notification.Member->AccountId.IsValid());
+
+		// 2.3 setup
+		{
+			EventCapture.Empty();
+
+			LobbyAttribute1Value.Set(TEXT("ModifiedAttribute1"));
+			LobbyMemberAttribute1Value.Set(TEXT("ModifiedMemberAttribute1"));
+			LobbySchemaName = TEXT("ModifiedSchemaName");
+
+			TMap<FOnlineAccountIdHandle, TSharedRef<FClientLobbyMemberSnapshot>> LobbyMemberSnapshots;
+			{
+				TSharedRef<FClientLobbyMemberSnapshot> MemberSnapshot = MakeShared<FClientLobbyMemberSnapshot>();
+				MemberSnapshot->AccountId = User1;
+				LobbyMemberSnapshots.Add(User1, MemberSnapshot);
+			}
+			{
+				TSharedRef<FClientLobbyMemberSnapshot> MemberSnapshot = MakeShared<FClientLobbyMemberSnapshot>();
+				MemberSnapshot->AccountId = User2;
+				MemberSnapshot->Attributes.Add(LobbyMemberAttribute1Key, LobbyMemberAttribute1Value);
+				LobbyMemberSnapshots.Add(User2, MemberSnapshot);
+			}
+
+			FClientLobbySnapshot LobbySnapshot;
+			LobbySnapshot.OwnerAccountId = User2;
+			LobbySnapshot.SchemaName = LobbySchemaName;
+			LobbySnapshot.MaxMembers = LobbyMaxMembers;
+			LobbySnapshot.JoinPolicy = LobbyJoinPolicy;
+			LobbySnapshot.Attributes.Add(LobbyAttribute1Key, LobbyAttribute1Value);
+			LobbySnapshot.Members = {User1, User2};
+
+			TMap<FOnlineAccountIdHandle, ELobbyMemberLeaveReason> LeaveReasons;
+			FApplyLobbyUpdateResult Result = ClientData.ApplyLobbyUpdateFromServiceSnapshot(MoveTemp(LobbySnapshot), MoveTemp(LobbyMemberSnapshots), MoveTemp(LeaveReasons), &LobbyEvents);
+			UTEST_EQUAL("Check no local members left.", Result.LeavingLocalMembers.Num(), 0);
+		}
+
+		// 2.3.1. Snapshot with changed lobby attribute triggers lobby attribute changed.
+		// 2.3.2. Snapshot with changed lobby member attribute triggers lobby member attribute changed.
+		// 2.3.3. Snapshot with changed lobby schema triggers lobby schema changed.
+		// 2.3.4. Snapshot with changed lobby leader triggers lobby leader changed.
+		UTEST_EQUAL("Check events received.", EventCapture.GetTotalNotificationsReceived(), 4);
+		UTEST_EQUAL("Check LobbyAttributesChanged is valid.", EventCapture.LobbyAttributesChanged.Num(), 1);
+		UTEST_EQUAL("Check LobbyAttributesChanged is valid.", EventCapture.LobbyAttributesChanged[0].Notification.Lobby, ClientData.GetPublicDataPtr());
+		UTEST_EQUAL("Check LobbyAttributesChanged is valid.", EventCapture.LobbyAttributesChanged[0].Notification.ChangedAttributes.Num(), 1);
+		UTEST_NOT_NULL("Check LobbyAttributesChanged is valid.", EventCapture.LobbyAttributesChanged[0].Notification.ChangedAttributes.Find(LobbyAttribute1Key));
+		UTEST_EQUAL("Check LobbyMemberAttributesChanged is valid.", EventCapture.LobbyMemberAttributesChanged.Num(), 1);
+		UTEST_EQUAL("Check LobbyMemberAttributesChanged is valid.", EventCapture.LobbyMemberAttributesChanged[0].Notification.Lobby, ClientData.GetPublicDataPtr());
+		UTEST_EQUAL("Check LobbyMemberAttributesChanged is valid.", EventCapture.LobbyMemberAttributesChanged[0].Notification.Member->AccountId, User2);
+		UTEST_EQUAL("Check LobbyMemberAttributesChanged is valid.", EventCapture.LobbyMemberAttributesChanged[0].Notification.ChangedAttributes.Num(), 1);
+		UTEST_NOT_NULL("Check LobbyMemberAttributesChanged is valid.", EventCapture.LobbyMemberAttributesChanged[0].Notification.ChangedAttributes.Find(LobbyMemberAttribute1Key));
+		UTEST_EQUAL("Check LobbySchemaChanged is valid.", EventCapture.LobbySchemaChanged.Num(), 1);
+		UTEST_EQUAL("Check LobbySchemaChanged is valid.", EventCapture.LobbySchemaChanged[0].Notification.Lobby, ClientData.GetPublicDataPtr());
+		UTEST_EQUAL("Check LobbyLeaderChanged is valid.", EventCapture.LobbyLeaderChanged.Num(), 1);
+		UTEST_EQUAL("Check LobbyLeaderChanged is valid.", EventCapture.LobbyLeaderChanged[0].Notification.Lobby, ClientData.GetPublicDataPtr());
+		UTEST_EQUAL("Check LobbyLeaderChanged is valid.", EventCapture.LobbyLeaderChanged[0].Notification.Leader->AccountId, User2);
+
+		// 2.4 setup
+		{
+			EventCapture.Empty();
+
+			TMap<FOnlineAccountIdHandle, TSharedRef<FClientLobbyMemberSnapshot>> LobbyMemberSnapshots;
+			{
+				TSharedRef<FClientLobbyMemberSnapshot> MemberSnapshot = MakeShared<FClientLobbyMemberSnapshot>();
+				MemberSnapshot->AccountId = User2;
+				LobbyMemberSnapshots.Add(User2, MemberSnapshot);
+			}
+
+			FClientLobbySnapshot LobbySnapshot;
+			LobbySnapshot.OwnerAccountId = User2;
+			LobbySnapshot.SchemaName = LobbySchemaName;
+			LobbySnapshot.MaxMembers = LobbyMaxMembers;
+			LobbySnapshot.JoinPolicy = LobbyJoinPolicy;
+			LobbySnapshot.Members = {User2};
+
+			TMap<FOnlineAccountIdHandle, ELobbyMemberLeaveReason> LeaveReasons;
+			LeaveReasons.Add(User1, ELobbyMemberLeaveReason::Kicked);
+
+			FApplyLobbyUpdateResult Result = ClientData.ApplyLobbyUpdateFromServiceSnapshot(MoveTemp(LobbySnapshot), MoveTemp(LobbyMemberSnapshots), MoveTemp(LeaveReasons), &LobbyEvents);
+			UTEST_EQUAL("Check local member left.", Result.LeavingLocalMembers.Num(), 1);
+			UTEST_EQUAL("Check local member left.", Result.LeavingLocalMembers[0], User1);
+		}
+
+		// 2.4.1. Snapshot with cleared lobby attribute triggers lobby attribute changed.
+		// 2.4.2. Snapshot with cleared lobby member attribute triggers lobby member attribute changed.
+		// 2.4.3. Snapshot with members leaving triggers lobby member left and lobby left.
+		UTEST_EQUAL("Check events received.", EventCapture.GetTotalNotificationsReceived(), 5);
+		UTEST_EQUAL("Check LobbyAttributesChanged is valid.", EventCapture.LobbyAttributesChanged.Num(), 1);
+		UTEST_EQUAL("Check LobbyAttributesChanged is valid.", EventCapture.LobbyAttributesChanged[0].Notification.Lobby, ClientData.GetPublicDataPtr());
+		UTEST_EQUAL("Check LobbyAttributesChanged is valid.", EventCapture.LobbyAttributesChanged[0].Notification.ChangedAttributes.Num(), 1);
+		UTEST_NOT_NULL("Check LobbyAttributesChanged is valid.", EventCapture.LobbyAttributesChanged[0].Notification.ChangedAttributes.Find(LobbyAttribute1Key));
+		UTEST_EQUAL("Check LobbyMemberAttributesChanged is valid.", EventCapture.LobbyMemberAttributesChanged.Num(), 1);
+		UTEST_EQUAL("Check LobbyMemberAttributesChanged is valid.", EventCapture.LobbyMemberAttributesChanged[0].Notification.Lobby, ClientData.GetPublicDataPtr());
+		UTEST_EQUAL("Check LobbyMemberAttributesChanged is valid.", EventCapture.LobbyMemberAttributesChanged[0].Notification.Member->AccountId, User2);
+		UTEST_EQUAL("Check LobbyMemberAttributesChanged is valid.", EventCapture.LobbyMemberAttributesChanged[0].Notification.ChangedAttributes.Num(), 1);
+		UTEST_NOT_NULL("Check LobbyMemberAttributesChanged is valid.", EventCapture.LobbyMemberAttributesChanged[0].Notification.ChangedAttributes.Find(LobbyMemberAttribute1Key));
+		UTEST_EQUAL("Check LobbyMemberLeft valid.", EventCapture.LobbyMemberLeft.Num(), 2);
+		FLobbyEventCapture::TNotificationInfo<FLobbyMemberLeft>* User1LeaveNotification = EventCapture.LobbyMemberLeft[0].Notification.Member->AccountId == User1 ? &EventCapture.LobbyMemberLeft[0] : &EventCapture.LobbyMemberLeft[1];
+		FLobbyEventCapture::TNotificationInfo<FLobbyMemberLeft>* User2LeaveNotification = EventCapture.LobbyMemberLeft[0].Notification.Member->AccountId == User2 ? &EventCapture.LobbyMemberLeft[0] : &EventCapture.LobbyMemberLeft[1];
+		UTEST_EQUAL("Check LobbyMemberLeft valid.", User1LeaveNotification->Notification.Lobby, ClientData.GetPublicDataPtr());
+		UTEST_EQUAL("Check LobbyMemberLeft valid.", User1LeaveNotification->Notification.Member->AccountId, User1);
+		UTEST_EQUAL("Check LobbyMemberLeft valid.", User1LeaveNotification->Notification.Reason, ELobbyMemberLeaveReason::Kicked);
+		UTEST_EQUAL("Check LobbyMemberLeft valid.", User2LeaveNotification->Notification.Lobby, ClientData.GetPublicDataPtr());
+		UTEST_EQUAL("Check LobbyMemberLeft valid.", User2LeaveNotification->Notification.Member->AccountId, User2);
+		UTEST_EQUAL("Check LobbyMemberLeft valid.", User2LeaveNotification->Notification.Reason, ELobbyMemberLeaveReason::Left);
+		UTEST_EQUAL("Check LobbyLeft valid.", EventCapture.LobbyLeft.Num(), 1);
+		UTEST_EQUAL("Check LobbyLeft valid.", EventCapture.LobbyLeft[0].Notification.Lobby, ClientData.GetPublicDataPtr());
+		UTEST_EQUAL("Check LobbyLeft ordering.", EventCapture.LobbyLeft[0].GlobalIndex, EventCapture.GetTotalNotificationsReceived() - 1);
+		FLobbyEventCapture::TNotificationInfo<FLobbyMemberLeft>* LastMemberLeaveNotification = User1LeaveNotification->GlobalIndex > User2LeaveNotification->GlobalIndex ? User1LeaveNotification : User2LeaveNotification;
+		UTEST_EQUAL("Check LobbyMemberLeft ordering.", LastMemberLeaveNotification->GlobalIndex, EventCapture.LobbyLeft[0].GlobalIndex - 1);
+	}
+
+	// 3. Snapshot supplemental
+	{
+		FLobbyEvents LobbyEvents;
+		FLobbyEventCapture EventCapture(LobbyEvents);
+		FOnlineLobbyIdHandle LobbyId(EOnlineServices::Null, 1);
+		FClientLobbyData ClientData(LobbyId);
+
+		FOnlineAccountIdHandle User1(EOnlineServices::Null, 1);
+		FOnlineAccountIdHandle User2(EOnlineServices::Null, 2);
+		FName LobbySchemaName = TEXT("SchemaName");
+		int32 LobbyMaxMembers = 5;
+		ELobbyJoinPolicy LobbyJoinPolicy = ELobbyJoinPolicy::PublicAdvertised;
+
+		// 3.1 setup - initialize lobby data with initial snapshot.
+		{
+			TMap<FOnlineAccountIdHandle, TSharedRef<FClientLobbyMemberSnapshot>> LobbyMemberSnapshots;
+			{
+				TSharedRef<FClientLobbyMemberSnapshot> MemberSnapshot = MakeShared<FClientLobbyMemberSnapshot>();
+				MemberSnapshot->AccountId = User1;
+				LobbyMemberSnapshots.Add(User1, MemberSnapshot);
+			}
+			{
+				TSharedRef<FClientLobbyMemberSnapshot> MemberSnapshot = MakeShared<FClientLobbyMemberSnapshot>();
+				MemberSnapshot->AccountId = User2;
+				LobbyMemberSnapshots.Add(User2, MemberSnapshot);
+			}
+
+			FClientLobbySnapshot LobbySnapshot;
+			LobbySnapshot.OwnerAccountId = User1;
+			LobbySnapshot.SchemaName = LobbySchemaName;
+			LobbySnapshot.MaxMembers = LobbyMaxMembers;
+			LobbySnapshot.JoinPolicy = LobbyJoinPolicy;
+			LobbySnapshot.Members = {User1, User2};
+
+			TMap<FOnlineAccountIdHandle, ELobbyMemberLeaveReason> LeaveReasons;
+			FApplyLobbyUpdateResult Result = ClientData.ApplyLobbyUpdateFromServiceSnapshot(MoveTemp(LobbySnapshot), MoveTemp(LobbyMemberSnapshots), MoveTemp(LeaveReasons), &LobbyEvents);
+			UTEST_EQUAL("Check no local members left.", Result.LeavingLocalMembers.Num(), 0);
+		}
+
+		// 3.1 setup continued. Add local member to enable notifications.
+		{
+			FClientLobbyDataChanges LobbyChanges;
+			LobbyChanges.MutatedMembers.Add(User1, MakeShared<FClientLobbyMemberDataChanges>());
+			FApplyLobbyUpdateResult Result = ClientData.ApplyLobbyUpdateFromLocalChanges(MoveTemp(LobbyChanges), LobbyEvents);
+			UTEST_EQUAL("Check no local members left.", Result.LeavingLocalMembers.Num(), 0);
+		}
+
+		// 3.1 setup continued. Change snapshot for the test.
+		{
+			EventCapture.Empty();
+
+			TMap<FOnlineAccountIdHandle, TSharedRef<FClientLobbyMemberSnapshot>> LobbyMemberSnapshots;
+			{
+				TSharedRef<FClientLobbyMemberSnapshot> MemberSnapshot = MakeShared<FClientLobbyMemberSnapshot>();
+				MemberSnapshot->AccountId = User1;
+				LobbyMemberSnapshots.Add(User1, MemberSnapshot);
+			}
+
+			FClientLobbySnapshot LobbySnapshot;
+			LobbySnapshot.OwnerAccountId = User1;
+			LobbySnapshot.SchemaName = LobbySchemaName;
+			LobbySnapshot.MaxMembers = LobbyMaxMembers;
+			LobbySnapshot.JoinPolicy = LobbyJoinPolicy;
+			LobbySnapshot.Members = {User1};
+
+			TMap<FOnlineAccountIdHandle, ELobbyMemberLeaveReason> LeaveReasons;
+			FApplyLobbyUpdateResult Result = ClientData.ApplyLobbyUpdateFromServiceSnapshot(MoveTemp(LobbySnapshot), MoveTemp(LobbyMemberSnapshots), MoveTemp(LeaveReasons), &LobbyEvents);
+			UTEST_EQUAL("Check no local members left.", Result.LeavingLocalMembers.Num(), 0);
+		}
+
+		// 3.1. Test that Member left is fired even when a leave reason is not specified.
+		UTEST_EQUAL("Check events received.", EventCapture.GetTotalNotificationsReceived(), 1);
+		UTEST_EQUAL("Check LobbyMemberLeft valid.", EventCapture.LobbyMemberLeft.Num(), 1);
+		UTEST_EQUAL("Check LobbyMemberLeft valid.", EventCapture.LobbyMemberLeft[0].Notification.Lobby, ClientData.GetPublicDataPtr());
+		UTEST_EQUAL("Check LobbyMemberLeft valid.", EventCapture.LobbyMemberLeft[0].Notification.Member->AccountId, User2);
+		UTEST_EQUAL("Check LobbyMemberLeft valid.", EventCapture.LobbyMemberLeft[0].Notification.Reason, ELobbyMemberLeaveReason::Disconnected);
+
+		// 3.2 setup
+		{
+			EventCapture.Empty();
+
+			TMap<FOnlineAccountIdHandle, TSharedRef<FClientLobbyMemberSnapshot>> LobbyMemberSnapshots;
+			{
+				TSharedRef<FClientLobbyMemberSnapshot> MemberSnapshot = MakeShared<FClientLobbyMemberSnapshot>();
+				MemberSnapshot->AccountId = User1;
+				LobbyMemberSnapshots.Add(User1, MemberSnapshot);
+			}
+			{
+				TSharedRef<FClientLobbyMemberSnapshot> MemberSnapshot = MakeShared<FClientLobbyMemberSnapshot>();
+				MemberSnapshot->AccountId = User2;
+				LobbyMemberSnapshots.Add(User2, MemberSnapshot);
+			}
+
+			FClientLobbySnapshot LobbySnapshot;
+			LobbySnapshot.OwnerAccountId = User1;
+			LobbySnapshot.SchemaName = LobbySchemaName;
+			LobbySnapshot.MaxMembers = LobbyMaxMembers;
+			LobbySnapshot.JoinPolicy = LobbyJoinPolicy;
+			LobbySnapshot.Members = {User1};
+
+			TMap<FOnlineAccountIdHandle, ELobbyMemberLeaveReason> LeaveReasons;
+			FApplyLobbyUpdateResult Result = ClientData.ApplyLobbyUpdateFromServiceSnapshot(MoveTemp(LobbySnapshot), MoveTemp(LobbyMemberSnapshots), MoveTemp(LeaveReasons), &LobbyEvents);
+			UTEST_EQUAL("Check no local members left.", Result.LeavingLocalMembers.Num(), 0);
+		}
+
+		// 3.2. Applying member snapshots for an account id not in the lobby members list does not fire lobby member joined.
+		UTEST_EQUAL("Check no events received.", EventCapture.GetTotalNotificationsReceived(), 0);
+
+		// 3.3 setup
+		{
+			EventCapture.Empty();
+
+			TMap<FOnlineAccountIdHandle, TSharedRef<FClientLobbyMemberSnapshot>> LobbyMemberSnapshots;
+			{
+				TSharedRef<FClientLobbyMemberSnapshot> MemberSnapshot = MakeShared<FClientLobbyMemberSnapshot>();
+				MemberSnapshot->AccountId = User1;
+				LobbyMemberSnapshots.Add(User1, MemberSnapshot);
+			}
+			{
+				TSharedRef<FClientLobbyMemberSnapshot> MemberSnapshot = MakeShared<FClientLobbyMemberSnapshot>();
+				MemberSnapshot->AccountId = User2;
+				LobbyMemberSnapshots.Add(User2, MemberSnapshot);
+			}
+
+			FClientLobbySnapshot LobbySnapshot;
+			LobbySnapshot.OwnerAccountId = User1;
+			LobbySnapshot.SchemaName = LobbySchemaName;
+			LobbySnapshot.MaxMembers = LobbyMaxMembers;
+			LobbySnapshot.JoinPolicy = LobbyJoinPolicy;
+			LobbySnapshot.Members = {User1, User2};
+
+			TMap<FOnlineAccountIdHandle, ELobbyMemberLeaveReason> LeaveReasons;
+			FApplyLobbyUpdateResult Result = ClientData.ApplyLobbyUpdateFromServiceSnapshot(MoveTemp(LobbySnapshot), MoveTemp(LobbyMemberSnapshots), MoveTemp(LeaveReasons), &LobbyEvents);
+			UTEST_EQUAL("Check no local members left.", Result.LeavingLocalMembers.Num(), 0);
+		}
+
+		// 3.3. Check lobby member joined notification received for remote member after local member joined.
+		UTEST_EQUAL("Check events received.", EventCapture.GetTotalNotificationsReceived(), 1);
+		UTEST_EQUAL("Check LobbyMemberJoined valid.", EventCapture.LobbyMemberJoined.Num(), 1);
+		UTEST_EQUAL("Check LobbyMemberJoined valid.", EventCapture.LobbyMemberJoined[0].Notification.Lobby, ClientData.GetPublicDataPtr());
+		UTEST_EQUAL("Check LobbyMemberJoined valid.", EventCapture.LobbyMemberJoined[0].Notification.Member->AccountId, User2);
+	}
+
+	// 4. Applying lobby changes produces expected notifications.
+	{
+		FLobbyEvents LobbyEvents;
+		FLobbyEventCapture EventCapture(LobbyEvents);
+		FOnlineLobbyIdHandle LobbyId(EOnlineServices::Null, 1);
+		FClientLobbyData ClientData(LobbyId);
+
+		FOnlineAccountIdHandle User1(EOnlineServices::Null, 1);
+		FOnlineAccountIdHandle User2(EOnlineServices::Null, 2);
+		FName LobbySchemaName = TEXT("SchemaName");
+		int32 LobbyMaxMembers = 5;
+		ELobbyJoinPolicy LobbyJoinPolicy = ELobbyJoinPolicy::PublicAdvertised;
+		FLobbyAttributeId LobbyAttribute1Key = TEXT("LobbyAttribute1");
+		FLobbyVariant LobbyAttribute1Value;
+		LobbyAttribute1Value.Set(TEXT("Attribute1"));
+		FLobbyAttributeId LobbyAttribute2Key = TEXT("LobbyAttribute2");
+		FLobbyVariant LobbyAttribute2Value;
+		LobbyAttribute2Value.Set(TEXT("Attribute2"));
+		FLobbyAttributeId LobbyMemberAttribute1Key = TEXT("LobbyMemberAttribute1");
+		FLobbyVariant LobbyMemberAttribute1Value;
+		LobbyMemberAttribute1Value.Set(TEXT("MemberAttribute1"));
+		FLobbyAttributeId LobbyMemberAttribute2Key = TEXT("LobbyMemberAttribute2");
+		FLobbyVariant LobbyMemberAttribute2Value;
+		LobbyMemberAttribute2Value.Set(TEXT("MemberAttribute2"));
+
+		// 4.1 setup. Local member in snapshot with remote member to simulate lobby join.
+		{
+			TMap<FOnlineAccountIdHandle, TSharedRef<FClientLobbyMemberSnapshot>> LobbyMemberSnapshots;
+			{
+				TSharedRef<FClientLobbyMemberSnapshot> MemberSnapshot = MakeShared<FClientLobbyMemberSnapshot>();
+				MemberSnapshot->AccountId = User1;
+				MemberSnapshot->Attributes.Add(LobbyMemberAttribute1Key, LobbyMemberAttribute1Value);
+				LobbyMemberSnapshots.Add(User1, MemberSnapshot);
+			}
+			{
+				TSharedRef<FClientLobbyMemberSnapshot> MemberSnapshot = MakeShared<FClientLobbyMemberSnapshot>();
+				MemberSnapshot->AccountId = User2;
+				LobbyMemberSnapshots.Add(User2, MemberSnapshot);
+			}
+
+			FClientLobbySnapshot LobbySnapshot;
+			LobbySnapshot.OwnerAccountId = User1;
+			LobbySnapshot.SchemaName = LobbySchemaName;
+			LobbySnapshot.MaxMembers = LobbyMaxMembers;
+			LobbySnapshot.JoinPolicy = LobbyJoinPolicy;
+			LobbySnapshot.Attributes.Add(LobbyAttribute1Key, LobbyAttribute1Value);
+			LobbySnapshot.Members = {User1, User2};
+
+			TMap<FOnlineAccountIdHandle, ELobbyMemberLeaveReason> LeaveReasons;
+			FApplyLobbyUpdateResult Result = ClientData.ApplyLobbyUpdateFromServiceSnapshot(MoveTemp(LobbySnapshot), MoveTemp(LobbyMemberSnapshots), MoveTemp(LeaveReasons), &LobbyEvents);
+			UTEST_EQUAL("Check no local members left.", Result.LeavingLocalMembers.Num(), 0);
+		}
+
+		UTEST_EQUAL("Check no events received.", EventCapture.GetTotalNotificationsReceived(), 0);
+
+		// 4.1 setup continued.
+		{
+			LobbyAttribute1Value.Set(TEXT("MutatedAttribute1"));
+			LobbySchemaName = TEXT("Beans");
+
+			FClientLobbyDataChanges LobbyChanges;
+			LobbyChanges.MutatedAttributes.Add(LobbyAttribute1Key, LobbyAttribute1Value);
+			LobbyChanges.LobbySchema = LobbySchemaName;
+			LobbyChanges.OwnerAccountId = User2;
+			FApplyLobbyUpdateResult Result = ClientData.ApplyLobbyUpdateFromLocalChanges(MoveTemp(LobbyChanges), LobbyEvents);
+			UTEST_EQUAL("Check no local members left.", Result.LeavingLocalMembers.Num(), 0);
+		}
+
+		// 4.1.1. Local changes with changed lobby attribute does not trigger lobby attribute changed.
+		// 4.1.2. Local changes with changed lobby schema does not trigger lobby schema changed.
+		// 4.1.3. Local changes with changed lobby leader does not trigger lobby leader changed.
+		UTEST_EQUAL("Check no events received.", EventCapture.GetTotalNotificationsReceived(), 0);
+
+		// 4.2 setup.
+		{
+			FClientLobbyDataChanges LobbyChanges;
+			LobbyChanges.MutatedMembers.Add(User1, MakeShared<FClientLobbyMemberDataChanges>());
+			FApplyLobbyUpdateResult Result = ClientData.ApplyLobbyUpdateFromLocalChanges(MoveTemp(LobbyChanges), LobbyEvents);
+			UTEST_EQUAL("Check no local members left.", Result.LeavingLocalMembers.Num(), 0);
+		}
+
+		// 4.2.1. Local members added to lobby through local change triggers lobby joined and lobby member joined. On member joined triggered for remote member.
+		UTEST_EQUAL("Check events received.", EventCapture.GetTotalNotificationsReceived(), 3);
+		UTEST_EQUAL("Check LobbyJoined valid.", EventCapture.LobbyJoined.Num(), 1);
+		UTEST_EQUAL("Check LobbyJoined valid.", EventCapture.LobbyJoined[0].GlobalIndex, 0);
+		UTEST_EQUAL("Check LobbyMemberJoined valid.", EventCapture.LobbyMemberJoined.Num(), 2);
+		FLobbyEventCapture::TNotificationInfo<FLobbyMemberJoined>* User1JoinNotification = EventCapture.LobbyMemberJoined[0].Notification.Member->AccountId == User1 ? &EventCapture.LobbyMemberJoined[0] : &EventCapture.LobbyMemberJoined[1];
+		FLobbyEventCapture::TNotificationInfo<FLobbyMemberJoined>* User2JoinNotification = EventCapture.LobbyMemberJoined[0].Notification.Member->AccountId == User2 ? &EventCapture.LobbyMemberJoined[0] : &EventCapture.LobbyMemberJoined[1];
+		UTEST_EQUAL("Check LobbyMemberJoined valid.", User1JoinNotification->Notification.Lobby, ClientData.GetPublicDataPtr());
+		UTEST_EQUAL("Check LobbyMemberJoined valid.", User1JoinNotification->Notification.Member->AccountId, User1);
+		UTEST_EQUAL("Check LobbyMemberJoined valid.", User1JoinNotification->Notification.Member->Attributes.Num(), 1);
+		UTEST_EQUAL("Check LobbyMemberJoined valid.", User2JoinNotification->Notification.Lobby, ClientData.GetPublicDataPtr());
+		UTEST_EQUAL("Check LobbyMemberJoined valid.", User2JoinNotification->Notification.Member->AccountId, User2);
+
+		// 4.3 setup
+		{
+			EventCapture.Empty();
+
+			LobbyAttribute1Value.Set(TEXT("DoubleMutatedAttribute1"));
+			LobbyMemberAttribute1Value.Set(TEXT("DoubleMutatedMemberAttribute1"));
+			LobbySchemaName = TEXT("Waffles");
+
+			TSharedRef<FClientLobbyMemberDataChanges> User1Changes = MakeShared<FClientLobbyMemberDataChanges>();
+			User1Changes->MutatedAttributes.Add(LobbyMemberAttribute1Key, LobbyMemberAttribute1Value);
+
+			FClientLobbyDataChanges LobbyChanges;
+			LobbyChanges.MutatedAttributes.Add(LobbyAttribute1Key, LobbyAttribute1Value);
+			LobbyChanges.MutatedMembers.Add(User1, User1Changes);
+			LobbyChanges.LobbySchema = LobbySchemaName;
+			LobbyChanges.OwnerAccountId = User1;
+			FApplyLobbyUpdateResult Result = ClientData.ApplyLobbyUpdateFromLocalChanges(MoveTemp(LobbyChanges), LobbyEvents);
+			UTEST_EQUAL("Check no local members left.", Result.LeavingLocalMembers.Num(), 0);
+		}
+
+		// 4.3.1. Local changes with changed lobby attribute triggers lobby attribute changed.
+		// 4.3.2. Local changes with changed lobby member attribute triggers lobby member attribute changed.
+		// 4.3.3. Local changes with changed lobby schema triggers lobby schema changed.
+		// 4.3.4. Local changes with changed lobby leader triggers lobby leader changed.
+		UTEST_EQUAL("Check events received.", EventCapture.GetTotalNotificationsReceived(), 4);
+		UTEST_EQUAL("Check LobbyAttributesChanged is valid.", EventCapture.LobbyAttributesChanged.Num(), 1);
+		UTEST_EQUAL("Check LobbyAttributesChanged is valid.", EventCapture.LobbyAttributesChanged[0].Notification.Lobby, ClientData.GetPublicDataPtr());
+		UTEST_EQUAL("Check LobbyAttributesChanged is valid.", EventCapture.LobbyAttributesChanged[0].Notification.ChangedAttributes.Num(), 1);
+		UTEST_NOT_NULL("Check LobbyAttributesChanged is valid.", EventCapture.LobbyAttributesChanged[0].Notification.ChangedAttributes.Find(LobbyAttribute1Key));
+		UTEST_EQUAL("Check LobbyMemberAttributesChanged is valid.", EventCapture.LobbyMemberAttributesChanged.Num(), 1);
+		UTEST_EQUAL("Check LobbyMemberAttributesChanged is valid.", EventCapture.LobbyMemberAttributesChanged[0].Notification.Lobby, ClientData.GetPublicDataPtr());
+		UTEST_EQUAL("Check LobbyMemberAttributesChanged is valid.", EventCapture.LobbyMemberAttributesChanged[0].Notification.Member->AccountId, User1);
+		UTEST_EQUAL("Check LobbyMemberAttributesChanged is valid.", EventCapture.LobbyMemberAttributesChanged[0].Notification.ChangedAttributes.Num(), 1);
+		UTEST_NOT_NULL("Check LobbyMemberAttributesChanged is valid.", EventCapture.LobbyMemberAttributesChanged[0].Notification.ChangedAttributes.Find(LobbyMemberAttribute1Key));
+		UTEST_EQUAL("Check LobbySchemaChanged is valid.", EventCapture.LobbySchemaChanged.Num(), 1);
+		UTEST_EQUAL("Check LobbySchemaChanged is valid.", EventCapture.LobbySchemaChanged[0].Notification.Lobby, ClientData.GetPublicDataPtr());
+		UTEST_EQUAL("Check LobbyLeaderChanged is valid.", EventCapture.LobbyLeaderChanged.Num(), 1);
+		UTEST_EQUAL("Check LobbyLeaderChanged is valid.", EventCapture.LobbyLeaderChanged[0].Notification.Lobby, ClientData.GetPublicDataPtr());
+		UTEST_EQUAL("Check LobbyLeaderChanged is valid.", EventCapture.LobbyLeaderChanged[0].Notification.Leader->AccountId, User1);
+
+		// 4.4 setup
+		{
+			EventCapture.Empty();
+
+			TSharedRef<FClientLobbyMemberDataChanges> User1Changes = MakeShared<FClientLobbyMemberDataChanges>();
+			User1Changes->ClearedAttributes.Add(LobbyMemberAttribute1Key);
+
+			FClientLobbyDataChanges LobbyChanges;
+			LobbyChanges.ClearedAttributes.Add(LobbyAttribute1Key);
+			LobbyChanges.MutatedMembers.Add(User1, User1Changes);
+			LobbyChanges.LeavingMembers.Add(User1, ELobbyMemberLeaveReason::Left);
+			FApplyLobbyUpdateResult Result = ClientData.ApplyLobbyUpdateFromLocalChanges(MoveTemp(LobbyChanges), LobbyEvents);
+			UTEST_EQUAL("Check local member left.", Result.LeavingLocalMembers.Num(), 1);
+			UTEST_EQUAL("Check local member left.", Result.LeavingLocalMembers[0], User1);
+		}
+
+		// 4.4.1. Local changes with cleared lobby attribute triggers lobby attribute changed.
+		// 4.4.2. Local changes with cleared member attribute triggers member attribute changed.
+		// 4.4.1. Local member removed through local change results in member leave for all members followed by lobby leave.
+		UTEST_EQUAL("Check events received.", EventCapture.GetTotalNotificationsReceived(), 5);
+		UTEST_EQUAL("Check LobbyAttributesChanged is valid.", EventCapture.LobbyAttributesChanged.Num(), 1);
+		UTEST_EQUAL("Check LobbyAttributesChanged is valid.", EventCapture.LobbyAttributesChanged[0].Notification.Lobby, ClientData.GetPublicDataPtr());
+		UTEST_EQUAL("Check LobbyAttributesChanged is valid.", EventCapture.LobbyAttributesChanged[0].Notification.ChangedAttributes.Num(), 1);
+		UTEST_NOT_NULL("Check LobbyAttributesChanged is valid.", EventCapture.LobbyAttributesChanged[0].Notification.ChangedAttributes.Find(LobbyAttribute1Key));
+		UTEST_EQUAL("Check LobbyMemberAttributesChanged is valid.", EventCapture.LobbyMemberAttributesChanged.Num(), 1);
+		UTEST_EQUAL("Check LobbyMemberAttributesChanged is valid.", EventCapture.LobbyMemberAttributesChanged[0].Notification.Lobby, ClientData.GetPublicDataPtr());
+		UTEST_EQUAL("Check LobbyMemberAttributesChanged is valid.", EventCapture.LobbyMemberAttributesChanged[0].Notification.Member->AccountId, User1);
+		UTEST_EQUAL("Check LobbyMemberAttributesChanged is valid.", EventCapture.LobbyMemberAttributesChanged[0].Notification.ChangedAttributes.Num(), 1);
+		UTEST_NOT_NULL("Check LobbyMemberAttributesChanged is valid.", EventCapture.LobbyMemberAttributesChanged[0].Notification.ChangedAttributes.Find(LobbyMemberAttribute1Key));
+		UTEST_EQUAL("Check LobbyMemberLeft valid.", EventCapture.LobbyMemberLeft.Num(), 2);
+		FLobbyEventCapture::TNotificationInfo<FLobbyMemberLeft>* User1LeaveNotification = EventCapture.LobbyMemberLeft[0].Notification.Member->AccountId == User1 ? &EventCapture.LobbyMemberLeft[0] : &EventCapture.LobbyMemberLeft[1];
+		FLobbyEventCapture::TNotificationInfo<FLobbyMemberLeft>* User2LeaveNotification = EventCapture.LobbyMemberLeft[0].Notification.Member->AccountId == User2 ? &EventCapture.LobbyMemberLeft[0] : &EventCapture.LobbyMemberLeft[1];
+		UTEST_EQUAL("Check LobbyMemberLeft valid.", User1LeaveNotification->Notification.Lobby, ClientData.GetPublicDataPtr());
+		UTEST_EQUAL("Check LobbyMemberLeft valid.", User1LeaveNotification->Notification.Member->AccountId, User1);
+		UTEST_EQUAL("Check LobbyMemberLeft valid.", User1LeaveNotification->Notification.Reason, ELobbyMemberLeaveReason::Left);
+		UTEST_EQUAL("Check LobbyMemberLeft valid.", User2LeaveNotification->Notification.Lobby, ClientData.GetPublicDataPtr());
+		UTEST_EQUAL("Check LobbyMemberLeft valid.", User2LeaveNotification->Notification.Member->AccountId, User2);
+		UTEST_EQUAL("Check LobbyMemberLeft valid.", User2LeaveNotification->Notification.Reason, ELobbyMemberLeaveReason::Left);
+		UTEST_EQUAL("Check LobbyLeft valid.", EventCapture.LobbyLeft.Num(), 1);
+		UTEST_EQUAL("Check LobbyLeft valid.", EventCapture.LobbyLeft[0].Notification.Lobby, ClientData.GetPublicDataPtr());
+		UTEST_EQUAL("Check LobbyLeft ordering.", EventCapture.LobbyLeft[0].GlobalIndex, EventCapture.GetTotalNotificationsReceived() - 1);
+		FLobbyEventCapture::TNotificationInfo<FLobbyMemberLeft>* LastMemberLeaveNotification = User1LeaveNotification->GlobalIndex > User2LeaveNotification->GlobalIndex ? User1LeaveNotification : User2LeaveNotification;
+		UTEST_EQUAL("Check LobbyMemberLeft ordering.", LastMemberLeaveNotification->GlobalIndex, EventCapture.LobbyLeft[0].GlobalIndex - 1);
+	}
+
+	// 5. Local changes supplemental
+	{
+		FLobbyEvents LobbyEvents;
+		FLobbyEventCapture EventCapture(LobbyEvents);
+		FOnlineLobbyIdHandle LobbyId(EOnlineServices::Null, 1);
+		FClientLobbyData ClientData(LobbyId);
+
+		FOnlineAccountIdHandle User1(EOnlineServices::Null, 1);
+		FOnlineAccountIdHandle User2(EOnlineServices::Null, 2);
+		FName LobbySchemaName = TEXT("SchemaName");
+		int32 LobbyMaxMembers = 5;
+		ELobbyJoinPolicy LobbyJoinPolicy = ELobbyJoinPolicy::PublicAdvertised;
+		FLobbyAttributeId LobbyAttribute1Key = TEXT("LobbyAttribute1");
+		FLobbyVariant LobbyAttribute1Value;
+		LobbyAttribute1Value.Set(TEXT("Attribute1"));
+		FLobbyAttributeId LobbyAttribute2Key = TEXT("LobbyAttribute2");
+		FLobbyVariant LobbyAttribute2Value;
+		LobbyAttribute2Value.Set(TEXT("Attribute2"));
+		FLobbyAttributeId LobbyMemberAttribute1Key = TEXT("LobbyMemberAttribute1");
+		FLobbyVariant LobbyMemberAttribute1Value;
+		LobbyMemberAttribute1Value.Set(TEXT("MemberAttribute1"));
+		FLobbyAttributeId LobbyMemberAttribute2Key = TEXT("LobbyMemberAttribute2");
+		FLobbyVariant LobbyMemberAttribute2Value;
+		LobbyMemberAttribute2Value.Set(TEXT("MemberAttribute2"));
+
+		// 5.1 setup.
+		{
+			TMap<FOnlineAccountIdHandle, TSharedRef<FClientLobbyMemberSnapshot>> LobbyMemberSnapshots;
+			{
+				TSharedRef<FClientLobbyMemberSnapshot> MemberSnapshot = MakeShared<FClientLobbyMemberSnapshot>();
+				MemberSnapshot->AccountId = User1;
+				MemberSnapshot->Attributes.Add(LobbyMemberAttribute1Key, LobbyMemberAttribute1Value);
+				LobbyMemberSnapshots.Add(User1, MemberSnapshot);
+			}
+			{
+				TSharedRef<FClientLobbyMemberSnapshot> MemberSnapshot = MakeShared<FClientLobbyMemberSnapshot>();
+				MemberSnapshot->AccountId = User2;
+				LobbyMemberSnapshots.Add(User2, MemberSnapshot);
+			}
+
+			FClientLobbySnapshot LobbySnapshot;
+			LobbySnapshot.OwnerAccountId = User1;
+			LobbySnapshot.SchemaName = LobbySchemaName;
+			LobbySnapshot.MaxMembers = LobbyMaxMembers;
+			LobbySnapshot.JoinPolicy = LobbyJoinPolicy;
+			LobbySnapshot.Attributes.Add(LobbyAttribute1Key, LobbyAttribute1Value);
+			LobbySnapshot.Members = {User1, User2};
+
+			TMap<FOnlineAccountIdHandle, ELobbyMemberLeaveReason> LeaveReasons;
+			FApplyLobbyUpdateResult Result = ClientData.ApplyLobbyUpdateFromServiceSnapshot(MoveTemp(LobbySnapshot), MoveTemp(LobbyMemberSnapshots), MoveTemp(LeaveReasons), &LobbyEvents);
+			UTEST_EQUAL("Check no local members left.", Result.LeavingLocalMembers.Num(), 0);
+		}
+
+		// 5.1 setup continued.
+		{
+			FClientLobbyDataChanges LobbyChanges;
+			LobbyChanges.MutatedMembers.Add(User1, MakeShared<FClientLobbyMemberDataChanges>());
+			FApplyLobbyUpdateResult Result = ClientData.ApplyLobbyUpdateFromLocalChanges(MoveTemp(LobbyChanges), LobbyEvents);
+			UTEST_EQUAL("Check no local members left.", Result.LeavingLocalMembers.Num(), 0);
+		}
+
+		// 5.1 setup continued.
+		{
+			EventCapture.Empty();
+
+			FClientLobbyDataChanges LobbyChanges;
+			LobbyChanges.LeavingMembers.Add(User2, ELobbyMemberLeaveReason::Kicked);
+			FApplyLobbyUpdateResult Result = ClientData.ApplyLobbyUpdateFromLocalChanges(MoveTemp(LobbyChanges), LobbyEvents);
+			UTEST_EQUAL("Check no local members left.", Result.LeavingLocalMembers.Num(), 0);
+		}
+
+		// 5.1 Test local lobby owner kicking remote member.
+		UTEST_EQUAL("Check events received.", EventCapture.GetTotalNotificationsReceived(), 1);
+		UTEST_EQUAL("Check LobbyMemberLeft valid.", EventCapture.LobbyMemberLeft.Num(), 1);
+		UTEST_EQUAL("Check LobbyMemberLeft valid.", EventCapture.LobbyMemberLeft[0].Notification.Lobby, ClientData.GetPublicDataPtr());
+		UTEST_EQUAL("Check LobbyMemberLeft valid.", EventCapture.LobbyMemberLeft[0].Notification.Member->AccountId, User2);
+		UTEST_EQUAL("Check LobbyMemberLeft valid.", EventCapture.LobbyMemberLeft[0].Notification.Reason, ELobbyMemberLeaveReason::Kicked);
+	}
+
+	return true;
+}
+
+//--------------------------------------------------------------------------------------------------
+// Functional testing
+//--------------------------------------------------------------------------------------------------
+
+#if LOBBIES_FUNCTIONAL_TEST_ENABLED
+namespace Private {
+
+template <typename SecondaryOpType, typename OpType, typename Function>
+TFunction<TFuture<void>(TOnlineAsyncOp<OpType>&, typename SecondaryOpType::Params&&)> ConsumeStepResult(TOnlineAsyncOp<OpType>& InAsyncOp, Function Func)
+{
+	return [Func](TOnlineAsyncOp<OpType>& InAsyncOp, typename SecondaryOpType::Params&& InParams) mutable
+	{
+		TPromise<void> Promise;
+		auto Future = Promise.GetFuture();
+
+		UE_LOG(LogTemp, Log, TEXT("[LobbiesFunctionalTest::ConsumeStepResult] Starting secondary operation %s."), SecondaryOpType::Name);
+
+		Func(MoveTempIfPossible(InParams))
+		.Then([Promise = MoveTemp(Promise), Op = InAsyncOp.AsShared()](TFuture<TOnlineResult<SecondaryOpType>>&& Future) mutable
+		{
+			if (Future.Get().IsError())
+			{
+				UE_LOG(LogTemp, Log, TEXT("[LobbiesFunctionalTest::ConsumeStepResult] Failed secondary operation %s. Error: %s"), SecondaryOpType::Name, *Future.Get().GetErrorValue().GetLogString());
+				Op->SetError(Errors::RequestFailure(MoveTempIfPossible(Future.Get().GetErrorValue())));
+			}
+			else
+			{
+				UE_LOG(LogTemp, Log, TEXT("[LobbiesFunctionalTest::ConsumeStepResult] Completed secondary operation %s."), SecondaryOpType::Name);
+			}
+			Promise.EmplaceValue();
+		});
+
+		return Future;
+	};
+}
+
+template <typename SecondaryOpType, typename OpType, typename Function>
+TFunction<TFuture<void>(TOnlineAsyncOp<OpType>&)> ConsumeStepResult(TOnlineAsyncOp<OpType>& InAsyncOp, Function Func, typename SecondaryOpType::Params&& InParams)
+{
+	return [Func, InParams = MoveTempIfPossible(InParams)](TOnlineAsyncOp<OpType>& InAsyncOp) mutable
+	{
+		TPromise<void> Promise;
+		auto Future = Promise.GetFuture();
+
+		UE_LOG(LogTemp, Log, TEXT("[LobbiesFunctionalTest::ConsumeStepResult] Starting secondary operation %s."), SecondaryOpType::Name);
+
+		Func(MoveTempIfPossible(InParams))
+		.Then([Promise = MoveTemp(Promise), Op = InAsyncOp.AsShared()](TFuture<TOnlineResult<SecondaryOpType>>&& Future) mutable
+		{
+			if (Future.Get().IsError())
+			{
+				UE_LOG(LogTemp, Log, TEXT("[LobbiesFunctionalTest::ConsumeStepResult] Failed secondary operation %s. Error: %s"), SecondaryOpType::Name, *Future.Get().GetErrorValue().GetLogString());
+				Op->SetError(Errors::RequestFailure(MoveTempIfPossible(Future.Get().GetErrorValue())));
+			}
+			else
+			{
+				UE_LOG(LogTemp, Log, TEXT("[LobbiesFunctionalTest::ConsumeStepResult] Completed secondary operation %s."), SecondaryOpType::Name);
+			}
+			Promise.EmplaceValue();
+		});
+
+		return Future;
+	};
+}
+
+template <typename SecondaryOpType, typename OpType, typename Function>
+TFunction<TFuture<void>(TOnlineAsyncOp<OpType>&, typename SecondaryOpType::Params&&)> CaptureStepResult(TOnlineAsyncOp<OpType>& InAsyncOp, const FString& ResultKey, Function Func)
+{
+	return [ResultKey, Func](TOnlineAsyncOp<OpType>& InAsyncOp, typename SecondaryOpType::Params&& InParams) mutable
+	{
+		TPromise<void> Promise;
+		auto Future = Promise.GetFuture();
+
+		UE_LOG(LogTemp, Log, TEXT("[LobbiesFunctionalTest::CaptureStepResult] Starting secondary operation %s."), SecondaryOpType::Name);
+
+		Func(MoveTempIfPossible(InParams))
+		.Then([Promise = MoveTemp(Promise), Op = InAsyncOp.AsShared(), ResultKey](TFuture<TOnlineResult<SecondaryOpType>>&& Future) mutable
+		{
+			if (Future.Get().IsError())
+			{
+				UE_LOG(LogTemp, Log, TEXT("[LobbiesFunctionalTest::CaptureStepResult] Failed secondary operation %s. Error: %s"), SecondaryOpType::Name, *Future.Get().GetErrorValue().GetLogString());
+				Op->SetError(Errors::RequestFailure(MoveTempIfPossible(Future.Get().GetErrorValue())));
+			}
+			else
+			{
+				UE_LOG(LogTemp, Log, TEXT("[LobbiesFunctionalTest::CaptureStepResult] Captured secondary operation %s as %s"), SecondaryOpType::Name, *ResultKey);
+				Op->Data.Set(ResultKey, MoveTempIfPossible(Future.Get().GetOkValue()));
+			}
+			Promise.EmplaceValue();
+		});
+
+		return Future;
+	};
+}
+
+template <typename SecondaryOpType, typename OpType, typename Function>
+TFunction<TFuture<void>(TOnlineAsyncOp<OpType>&)> CaptureStepResult(TOnlineAsyncOp<OpType>& InAsyncOp, const FString& ResultKey, Function Func, typename SecondaryOpType::Params&& InParams)
+{
+	return [ResultKey, Func, InParams = MoveTempIfPossible(InParams)](TOnlineAsyncOp<OpType>& InAsyncOp) mutable
+	{
+		TPromise<void> Promise;
+		auto Future = Promise.GetFuture();
+
+		UE_LOG(LogTemp, Log, TEXT("[LobbiesFunctionalTest::CaptureStepResult] Starting secondary operation %s."), SecondaryOpType::Name);
+
+		Func(MoveTempIfPossible(InParams))
+		.Then([Promise = MoveTemp(Promise), Op = InAsyncOp.AsShared(), ResultKey](TFuture<TOnlineResult<SecondaryOpType>>&& Future) mutable
+		{
+			if (Future.Get().IsError())
+			{
+				UE_LOG(LogTemp, Log, TEXT("[LobbiesFunctionalTest::CaptureStepResult] Failed secondary operation %s. Error: %s"), SecondaryOpType::Name, *Future.Get().GetErrorValue().GetLogString());
+				Op->SetError(Errors::RequestFailure(MoveTempIfPossible(Future.Get().GetErrorValue())));
+			}
+			else
+			{
+				UE_LOG(LogTemp, Log, TEXT("[LobbiesFunctionalTest::CaptureStepResult] Captured secondary operation %s as %s"), SecondaryOpType::Name, *ResultKey);
+				Op->Data.Set(ResultKey, MoveTempIfPossible(Future.Get().GetOkValue()));
+			}
+			Promise.EmplaceValue();
+		});
+
+		return Future;
+	};
+}
+
+template <typename SecondaryOpType, typename OpType, typename Function>
+TFunction<TFuture<void>(TOnlineAsyncOp<OpType>&, typename SecondaryOpType::Params&&)> ConsumeOperationStepResult(TOnlineAsyncOp<OpType>& InAsyncOp, Function Func)
+{
+	return [Func](TOnlineAsyncOp<OpType>& InAsyncOp, typename SecondaryOpType::Params&& InParams) mutable -> TFuture<void>
+	{
+		auto Promise = MakeShared<TPromise<void>>();
+		auto Future = Promise->GetFuture();
+
+		UE_LOG(LogTemp, Log, TEXT("[LobbiesFunctionalTest::ConsumeOperationStepResult] Starting secondary operation %s."), SecondaryOpType::Name);
+
+		Func(MoveTempIfPossible(InParams))
+		.OnComplete([Promise, Op = InAsyncOp.AsShared()](const TOnlineResult<SecondaryOpType>& Result) mutable -> void
+		{
+			if (Result.IsError())
+			{
+				UE_LOG(LogTemp, Log, TEXT("[LobbiesFunctionalTest::ConsumeOperationStepResult] Failed secondary operation %s. Error: %s"), SecondaryOpType::Name, *Result.GetErrorValue().GetLogString());
+				Op->SetError(Errors::RequestFailure(MoveTempIfPossible(Result.GetErrorValue())));
+			}
+			else
+			{
+				UE_LOG(LogTemp, Log, TEXT("[LobbiesFunctionalTest::ConsumeOperationStepResult] Completed secondary operation %s."), SecondaryOpType::Name);
+			}
+			Promise->EmplaceValue();
+		});
+
+		return Future;
+	};
+}
+
+template <typename SecondaryOpType, typename OpType, typename Function>
+TFunction<TFuture<void>(TOnlineAsyncOp<OpType>&)> ConsumeOperationStepResult(TOnlineAsyncOp<OpType>& InAsyncOp, Function Func, typename SecondaryOpType::Params&& InParams)
+{
+	return [Func, InParams = MoveTempIfPossible(InParams)](TOnlineAsyncOp<OpType>& InAsyncOp) mutable
+	{
+		auto Promise = MakeShared<TPromise<void>>();
+		auto Future = Promise->GetFuture();
+
+		UE_LOG(LogTemp, Log, TEXT("[LobbiesFunctionalTest::ConsumeOperationStepResult] Starting secondary operation %s."), SecondaryOpType::Name);
+
+		Func(MoveTempIfPossible(InParams))
+		.OnComplete([Promise, Op = InAsyncOp.AsShared()](const TOnlineResult<SecondaryOpType>& Result) mutable
+		{
+			if (Result.IsError())
+			{
+				UE_LOG(LogTemp, Log, TEXT("[LobbiesFunctionalTest::ConsumeOperationStepResult] Failed secondary operation %s. Error: %s"), SecondaryOpType::Name, *Result.GetErrorValue().GetLogString());
+				Op->SetError(Errors::RequestFailure(MoveTempIfPossible(Result.GetErrorValue())));
+			}
+			else
+			{
+				UE_LOG(LogTemp, Log, TEXT("[LobbiesFunctionalTest::ConsumeOperationStepResult] Completed secondary operation %s."), SecondaryOpType::Name);
+			}
+			Promise->EmplaceValue();
+		});
+
+		return Future;
+	};
+}
+
+template <typename SecondaryOpType, typename OpType, typename Function>
+TFunction<TFuture<void>(TOnlineAsyncOp<OpType>&, typename SecondaryOpType::Params&&)> CaptureOperationStepResult(TOnlineAsyncOp<OpType>& InAsyncOp, const FString& ResultKey, Function Func)
+{
+	return [ResultKey, Func](TOnlineAsyncOp<OpType>& InAsyncOp, typename SecondaryOpType::Params&& InParams) mutable -> TFuture<void>
+	{
+		auto Promise = MakeShared<TPromise<void>>();
+		auto Future = Promise->GetFuture();
+
+		UE_LOG(LogTemp, Log, TEXT("[LobbiesFunctionalTest::CaptureOperationStepResult] Starting secondary operation %s."), SecondaryOpType::Name);
+
+		Func(MoveTempIfPossible(InParams))
+		.OnComplete([Promise, Op = InAsyncOp.AsShared(), ResultKey](const TOnlineResult<SecondaryOpType>& Result) mutable -> void
+		{
+			if (Result.IsError())
+			{
+				UE_LOG(LogTemp, Log, TEXT("[LobbiesFunctionalTest::CaptureOperationStepResult] Failed secondary operation %s. Error: %s"), SecondaryOpType::Name, *Result.GetErrorValue().GetLogString());
+				Op->SetError(Errors::RequestFailure(MoveTempIfPossible(Result.GetErrorValue())));
+			}
+			else
+			{
+				UE_LOG(LogTemp, Log, TEXT("[LobbiesFunctionalTest::CaptureOperationStepResult] Captured secondary operation %s as %s"), SecondaryOpType::Name, *ResultKey);
+				Op->Data.Set(ResultKey, Result.GetOkValue());
+			}
+			Promise->EmplaceValue();
+		});
+
+		return Future;
+	};
+}
+
+template <typename SecondaryOpType, typename OpType, typename Function>
+TFunction<TFuture<void>(TOnlineAsyncOp<OpType>&)> CaptureOperationStepResult(TOnlineAsyncOp<OpType>& InAsyncOp, const FString& ResultKey, Function Func, typename SecondaryOpType::Params&& InParams)
+{
+	return [ResultKey, Func, InParams = MoveTempIfPossible(InParams)](TOnlineAsyncOp<OpType>& InAsyncOp) mutable
+	{
+		auto Promise = MakeShared<TPromise<void>>();
+		auto Future = Promise->GetFuture();
+
+		UE_LOG(LogTemp, Log, TEXT("[LobbiesFunctionalTest::CaptureOperationStepResult] Starting secondary operation %s."), SecondaryOpType::Name);
+
+		Func(MoveTempIfPossible(InParams))
+		.OnComplete([Promise, Op = InAsyncOp.AsShared(), ResultKey](const TOnlineResult<SecondaryOpType>& Result) mutable
+		{
+			if (Result.IsError())
+			{
+				UE_LOG(LogTemp, Log, TEXT("[LobbiesFunctionalTest::CaptureOperationStepResult] Failed secondary operation %s. Error: %s"), SecondaryOpType::Name, *Result.GetErrorValue().GetLogString());
+				Op->SetError(Errors::RequestFailure(MoveTempIfPossible(Result.GetErrorValue())));
+			}
+			else
+			{
+				UE_LOG(LogTemp, Log, TEXT("[LobbiesFunctionalTest::CaptureOperationStepResult] Captured secondary operation %s as %s"), SecondaryOpType::Name, *ResultKey);
+				Op->Data.Set(ResultKey, Result.GetOkValue());
+			}
+			Promise->EmplaceValue();
+		});
+
+		return Future;
+	};
+}
+
+// todo: TFuture<void> does not work as expected with "Then".
+TFuture<int> AwaitNextGameTick()
+{
+	TSharedRef<TPromise<int>> Promise = MakeShared<TPromise<int>>();
+	TFuture<int> Future = Promise->GetFuture();
+
+	FTSTicker::GetCoreTicker().AddTicker(
+	TEXT("LobbiesFunctionalTest::AwaitNextGameTick"),
+	0.f,
+	[Promise](float)
+	{
+		Promise->EmplaceValue();
+		return false;
+	});
+
+	return Future;
+}
+
+TFuture<TDefaultErrorResultInternal<FOnlineLobbyIdHandle>> AwaitInvitation(
+	FLobbyEvents& LobbyEvents,
+	FOnlineAccountIdHandle TargetAccountId,
+	FOnlineLobbyIdHandle LobbyId,
+	float TimeoutSeconds)
+{
+	// todo: Make this nicer - current implementation has issue on shutdown.
+	struct AwaitInvitationState
+	{
+		TPromise<TDefaultErrorResultInternal<FOnlineLobbyIdHandle>> Promise;
+		FOnlineEventDelegateHandle OnLobbyInvitationAddedHandle;
+		FTSTicker::FDelegateHandle OnAwaitExpiredHandle;
+		bool bSignaled = false;
+	};
+
+	TSharedRef<AwaitInvitationState> AwaitState = MakeShared<AwaitInvitationState>();
+
+	AwaitState->OnLobbyInvitationAddedHandle = LobbyEvents.OnLobbyInvitationAdded.Add(
+	[TargetAccountId, LobbyId, AwaitState](const FLobbyInvitationAdded& Notification)
+	{
+		FOnlineEventDelegateHandle DelegateHandle = MoveTemp(AwaitState->OnLobbyInvitationAddedHandle);
+
+		if (!AwaitState->bSignaled)
+		{
+			if (Notification.Lobby->LobbyId == LobbyId && Notification.LocalUserId == TargetAccountId)
+			{
+				FTSTicker::GetCoreTicker().RemoveTicker(AwaitState->OnAwaitExpiredHandle);
+				AwaitState->bSignaled = true;
+
+				AwaitNextGameTick()
+				.Then([AwaitState, LobbyId](TFuture<int>&&)
+				{
+					AwaitState->Promise.EmplaceValue(LobbyId);
+				});
+			}
+		}
+	});
+
+	AwaitState->OnAwaitExpiredHandle = FTSTicker::GetCoreTicker().AddTicker(
+	TEXT("LobbiesFunctionalTest::AwaitInvitation"),
+	TimeoutSeconds,
+	[AwaitState](float)
+	{
+		if (!AwaitState->bSignaled)
+		{
+			// Todo: Errors.
+			AwaitState->bSignaled = true;
+			AwaitState->OnLobbyInvitationAddedHandle.Unbind();
+			AwaitState->Promise.EmplaceValue(Errors::NotImplemented());
+			FTSTicker::GetCoreTicker().RemoveTicker(AwaitState->OnAwaitExpiredHandle);
+		}
+
+		return false;
+	});
+
+	return AwaitState->Promise.GetFuture();
+}
+
+struct FFunctionalTestLoginUser
+{
+	static constexpr TCHAR Name[] = TEXT("FunctionalTestLoginUser");
+
+	struct Params
+	{
+		IAuth* AuthInterface = nullptr;
+		FPlatformUserId PlatformUserId;
+		FString Type;
+		FString Id;
+		FString Token;
+	};
+
+	struct Result
+	{
+		TSharedPtr<FAccountInfo> AccountInfo;
+	};
+};
+
+struct FFunctionalTestLogoutUser
+{
+	static constexpr TCHAR Name[] = TEXT("FunctionalTestLogoutUser");
+
+	struct Params
+	{
+		IAuth* AuthInterface = nullptr;
+		FPlatformUserId PlatformUserId;
+	};
+
+	struct Result
+	{
+	};
+};
+
+struct FFunctionalTestLogoutAllUsers
+{
+	static constexpr TCHAR Name[] = TEXT("FunctionalTestLogoutAllUsers");
+
+	struct Params
+	{
+		IAuth* AuthInterface = nullptr;
+	};
+
+	struct Result
+	{
+	};
+};
+
+TFuture<TOnlineResult<FFunctionalTestLoginUser>> FunctionalTestLoginUser(FFunctionalTestLoginUser::Params&& Params)
+{
+	if (!Params.AuthInterface)
+	{
+		return MakeFulfilledPromise<TOnlineResult<FFunctionalTestLoginUser>>(Errors::MissingInterface()).GetFuture();
+	}
+
+	TSharedRef<TPromise<TOnlineResult<FFunctionalTestLoginUser>>> Promise = MakeShared<TPromise<TOnlineResult<FFunctionalTestLoginUser>>>();
+	TFuture<TOnlineResult<FFunctionalTestLoginUser>> Future = Promise->GetFuture();
+
+	FAuthLogin::Params LoginParams;
+	LoginParams.PlatformUserId = Params.PlatformUserId;
+	LoginParams.CredentialsType = Params.Type;
+	LoginParams.CredentialsId = Params.Id;
+	LoginParams.CredentialsToken.Set<FString>(Params.Token);
+
+	Params.AuthInterface->Login(MoveTemp(LoginParams))
+	.OnComplete([Promise, PlatformUserId = Params.PlatformUserId](const TOnlineResult<FAuthLogin>& LoginResult)
+	{
+		if (LoginResult.IsError())
+		{
+			Promise->EmplaceValue(Errors::RequestFailure());
+		}
+		else
+		{
+			Promise->EmplaceValue(FFunctionalTestLoginUser::Result{LoginResult.GetOkValue().AccountInfo});
+		}
+	});
+
+	return Future;
+}
+
+TFuture<TOnlineResult<FFunctionalTestLogoutUser>> FunctionalTestLogoutUser(FFunctionalTestLogoutUser::Params&& Params)
+{
+	if (!Params.AuthInterface)
+	{
+		return MakeFulfilledPromise<TOnlineResult<FFunctionalTestLogoutUser>>(Errors::MissingInterface()).GetFuture();
+	}
+
+	FAuthGetAccountByPlatformUserId::Params GetAccountParams;
+	GetAccountParams.PlatformUserId = Params.PlatformUserId;
+	TOnlineResult<FAuthGetAccountByPlatformUserId> Result = Params.AuthInterface->GetAccountByPlatformUserId(MoveTemp(GetAccountParams));
+	if (Result.IsError())
+	{
+		// Ignore errors for now. This function should ignore logged out users.
+		if (Result.GetErrorValue() == Errors::Unknown())
+		{
+			return MakeFulfilledPromise<TOnlineResult<FFunctionalTestLogoutUser>>(FFunctionalTestLogoutUser::Result{}).GetFuture();
+		}
+		else
+		{
+			return MakeFulfilledPromise<TOnlineResult<FFunctionalTestLogoutUser>>(Errors::Unknown(MoveTemp(Result.GetErrorValue()))).GetFuture();
+		}
+	}
+
+	TSharedRef<FAccountInfo> AccountInfo = Result.GetOkValue().AccountInfo;
+	if (AccountInfo->LoginStatus == ELoginStatus::NotLoggedIn)
+	{
+		return MakeFulfilledPromise<TOnlineResult<FFunctionalTestLogoutUser>>(Errors::InvalidUser()).GetFuture();
+	}
+
+	TSharedRef<TPromise<TOnlineResult<FFunctionalTestLogoutUser>>> Promise = MakeShared<TPromise<TOnlineResult<FFunctionalTestLogoutUser>>>();
+	TFuture<TOnlineResult<FFunctionalTestLogoutUser>> Future = Promise->GetFuture();
+
+	FAuthLogout::Params LogoutParams;
+	LogoutParams.LocalUserId = AccountInfo->UserId;
+
+	Params.AuthInterface->Logout(MoveTemp(LogoutParams))
+	.OnComplete([Promise, PlatformUserId = Params.PlatformUserId](const TOnlineResult<FAuthLogout>& LogoutResult)
+	{
+		if (LogoutResult.IsError())
+		{
+			Promise->EmplaceValue(Errors::RequestFailure());
+		}
+		else
+		{
+			Promise->EmplaceValue(FFunctionalTestLogoutUser::Result{});
+		}
+	});
+
+	return Future;
+}
+
+TFuture<TOnlineResult<FFunctionalTestLogoutAllUsers>> FunctionalTestLogoutAllUsers(FFunctionalTestLogoutAllUsers::Params&& Params)
+{
+	TArray<TFuture<TOnlineResult<FFunctionalTestLogoutUser>>> LogoutFutures;
+	for (int32 index = 0; index < MAX_LOCAL_PLAYERS; ++index)
+	{
+		LogoutFutures.Emplace(FunctionalTestLogoutUser(FFunctionalTestLogoutUser::Params{Params.AuthInterface, FPlatformMisc::GetPlatformUserForUserIndex(index)}));
+	}
+
+	TSharedRef<TPromise<TOnlineResult<FFunctionalTestLogoutAllUsers>>> Promise = MakeShared<TPromise<TOnlineResult<FFunctionalTestLogoutAllUsers>>>();
+	TFuture<TOnlineResult<FFunctionalTestLogoutAllUsers>> Future = Promise->GetFuture();
+
+	WhenAll(MoveTempIfPossible(LogoutFutures))
+	.Then([Promise = MoveTemp(Promise)](TFuture<TArray<TOnlineResult<FFunctionalTestLogoutUser>>>&& Results)
+	{
+		bool HasAnyError = false;
+		for (TOnlineResult<FFunctionalTestLogoutUser>& Result : Results.Get())
+		{
+			HasAnyError |= Result.IsError();
+		}
+
+		if (HasAnyError)
+		{
+			Promise->EmplaceValue(Errors::RequestFailure());
+		}
+		else
+		{
+			Promise->EmplaceValue(FFunctionalTestLogoutAllUsers::Result{});
+		}
+	});
+
+	return Future;
+}
+
+template <typename OperationType>
+class FBindOperation
+{
+private:
+	class FBindImplBase
+	{
+	public:
+		virtual ~FBindImplBase() = default;
+		virtual TOnlineAsyncOpHandle<OperationType> operator()(typename OperationType::Params&& Params) const = 0;
+	};
+
+	template <typename ObjectType, typename Callable>
+	class FBindImpl : public FBindImplBase
+	{
+	public:
+		FBindImpl(ObjectType* Object, Callable Func)
+			: Object(Object)
+			, Func(Func)
+		{
+		}
+
+		virtual ~FBindImpl() = default;
+
+		virtual TOnlineAsyncOpHandle<OperationType> operator()(typename OperationType::Params&& Params) const
+		{
+			return (Object->*Func)(MoveTempIfPossible(Params));
+		}
+
+	private:
+		ObjectType* Object;
+		Callable Func;
+	};
+
+public:
+	template <typename ObjectType, typename Callable>
+	FBindOperation(ObjectType* Object, Callable Func)
+		: Impl(MakeShared<FBindImpl<ObjectType, Callable>>(Object, MoveTempIfPossible(Func)))
+	{
+	}
+
+	TOnlineAsyncOpHandle<OperationType> operator()(typename OperationType::Params&& Params) const
+	{
+		check(Impl);
+		return (*Impl)(MoveTempIfPossible(Params));
+	}
+
+private:
+	TSharedPtr<FBindImplBase> Impl;
+};
+
+} // Private
+
+template <typename DataType, typename OpType>
+const DataType& GetOpDataChecked(const TOnlineAsyncOp<OpType>& Op, const FString& Key)
+{
+	const DataType* Data = Op.Data.template Get<DataType>(Key);
+	check(Data);
+	return *Data;
+}
+
+struct FFunctionalTestConfig
+{
+	FString TestAccount1Type;
+	FString TestAccount1Id;
+	FString TestAccount1Token;
+
+	FString TestAccount2Type;
+	FString TestAccount2Id;
+	FString TestAccount2Token;
+
+	float InvitationWaitSeconds = 10.f;
+};
+
+namespace Meta {
+
+BEGIN_ONLINE_STRUCT_META(FFunctionalTestConfig)
+	ONLINE_STRUCT_FIELD(FFunctionalTestConfig, TestAccount1Type),
+	ONLINE_STRUCT_FIELD(FFunctionalTestConfig, TestAccount1Id),
+	ONLINE_STRUCT_FIELD(FFunctionalTestConfig, TestAccount1Token),
+
+	ONLINE_STRUCT_FIELD(FFunctionalTestConfig, TestAccount2Type),
+	ONLINE_STRUCT_FIELD(FFunctionalTestConfig, TestAccount2Id),
+	ONLINE_STRUCT_FIELD(FFunctionalTestConfig, TestAccount2Token)
+END_ONLINE_STRUCT_META()
+
+/* Meta */ }
+
+TOnlineAsyncOpHandle<FFunctionalTestLobbies> RunLobbyFunctionalTest(IAuth& AuthInterface, FLobbiesCommon& LobbiesCommon, FLobbyEvents& LobbyEvents)
+{
+	static const FString ConfigName = TEXT("FunctionalTest");
+	static const FString User1KeyName = TEXT("User1");
+	static const FString User2KeyName = TEXT("User2");
+	static const FString CreateLobbyKeyName = TEXT("CreateLobby");
+	static const FString FindLobbyKeyName = TEXT("FindLobby");
+	static const FString LobbyEventCaptureKeyName = TEXT("LobbyEventCapture");
+	static const FString ConfigNameKeyName = TEXT("Config");
+
+	TSharedRef<FFunctionalTestConfig> TestConfig = MakeShared<FFunctionalTestConfig>();
+	LobbiesCommon.LoadConfig(*TestConfig, ConfigName);
+
+	TOnlineAsyncOpRef<FFunctionalTestLobbies> Op = LobbiesCommon.GetOp<FFunctionalTestLobbies>(FFunctionalTestLobbies::Params{});
+
+	// Set up event capturing.
+	Op->Data.Set(LobbyEventCaptureKeyName, MakeShared<FLobbyEventCapture>(LobbyEvents));
+
+	Op->Then(Private::ConsumeStepResult<Private::FFunctionalTestLogoutAllUsers>(*Op, &Private::FunctionalTestLogoutAllUsers, Private::FFunctionalTestLogoutAllUsers::Params{&AuthInterface}))
+	.Then(Private::CaptureStepResult<Private::FFunctionalTestLoginUser>(*Op, User1KeyName, &Private::FunctionalTestLoginUser, Private::FFunctionalTestLoginUser::Params{&AuthInterface, FPlatformMisc::GetPlatformUserForUserIndex(0), TestConfig->TestAccount1Type, TestConfig->TestAccount1Id, TestConfig->TestAccount1Token}))
+	.Then(Private::CaptureStepResult<Private::FFunctionalTestLoginUser>(*Op, User2KeyName, &Private::FunctionalTestLoginUser, Private::FFunctionalTestLoginUser::Params{&AuthInterface, FPlatformMisc::GetPlatformUserForUserIndex(1), TestConfig->TestAccount2Type, TestConfig->TestAccount2Id, TestConfig->TestAccount2Token}))
+
+	//----------------------------------------------------------------------------------------------
+	// Test 1:
+	//    Step 1: Create a lobby with primary user.
+	//    Step 2: Modify lobby attribute.
+	//    Step 3: Modify lobby member attribute.
+	//    Step 4: Leave lobby with primary user.
+	//----------------------------------------------------------------------------------------------
+
+	.Then([](TOnlineAsyncOp<FFunctionalTestLobbies>& InAsyncOp)
+	{
+		// Prepare for next test check.
+		GetOpDataChecked<TSharedRef<FLobbyEventCapture>>(InAsyncOp, LobbyEventCaptureKeyName)->Empty();
+
+		const Private::FFunctionalTestLoginUser::Result& User1Info = GetOpDataChecked<Private::FFunctionalTestLoginUser::Result>(InAsyncOp, User1KeyName);
+
+		// Create a lobby
+		FCreateLobby::Params Params;
+		Params.LocalUserId = User1Info.AccountInfo->UserId;
+		Params.LocalName = TEXT("test");
+		Params.SchemaName = TEXT("test");
+		Params.MaxMembers = 2;
+		Params.JoinPolicy = ELobbyJoinPolicy::InvitationOnly;
+		Params.Attributes.Add(TEXT("test_attribute_key"), TEXT("test_attribute_value"));
+		Params.LocalUsers.Emplace(FJoinLobbyLocalUserData{User1Info.AccountInfo->UserId, {{TEXT("test_attribute_key"), TEXT("test_attribute_value")}}});
+		return Params;
+	})
+	.Then(Private::CaptureOperationStepResult<FCreateLobby>(*Op, CreateLobbyKeyName, Private::FBindOperation<FCreateLobby>(&LobbiesCommon, &FLobbiesCommon::CreateLobby)))
+	.Then([](TOnlineAsyncOp<FFunctionalTestLobbies>& InAsyncOp)
+	{
+		const Private::FFunctionalTestLoginUser::Result& User1Info = GetOpDataChecked<Private::FFunctionalTestLoginUser::Result>(InAsyncOp, User1KeyName);
+		const FCreateLobby::Result& CreateResult = GetOpDataChecked<FCreateLobby::Result>(InAsyncOp, CreateLobbyKeyName);
+
+		// Check both lobby joined events were received and in the correct order.
+		const TSharedRef<FLobbyEventCapture>& EventCapture = GetOpDataChecked<TSharedRef<FLobbyEventCapture>>(InAsyncOp, LobbyEventCaptureKeyName);
+		if (EventCapture->GetTotalNotificationsReceived() != 2
+			|| EventCapture->LobbyJoined.Num() != 1 || EventCapture->LobbyMemberJoined.Num() != 1
+			|| EventCapture->LobbyJoined[0].GlobalIndex > EventCapture->LobbyMemberJoined[0].GlobalIndex)
+		{
+			InAsyncOp.SetError(Errors::Cancelled());
+		}
+
+		// Check lobby attributes are set to the expected values.
+		if (!CreateResult.Lobby->Attributes.OrderIndependentCompareEqual(
+			TMap<FLobbyAttributeId, FLobbyVariant>{{TEXT("test_attribute_key"), TEXT("test_attribute_value")}}))
+		{
+			InAsyncOp.SetError(Errors::Cancelled());
+		}
+
+		const TSharedRef<const FLobbyMember>* MemberData = CreateResult.Lobby->Members.Find(User1Info.AccountInfo->UserId);
+		if (CreateResult.Lobby->Members.Num() == 1 && MemberData != nullptr)
+		{
+			if (!(**MemberData).Attributes.OrderIndependentCompareEqual(
+				TMap<FLobbyAttributeId, FLobbyVariant>{{TEXT("test_attribute_key"), TEXT("test_attribute_value")}}))
+			{
+				InAsyncOp.SetError(Errors::Cancelled());
+			}
+		}
+		else
+		{
+			InAsyncOp.SetError(Errors::Cancelled());
+		}
+	})
+	.Then([](TOnlineAsyncOp<FFunctionalTestLobbies>& InAsyncOp)
+	{
+		// Prepare for next test check.
+		GetOpDataChecked<TSharedRef<FLobbyEventCapture>>(InAsyncOp, LobbyEventCaptureKeyName)->Empty();
+
+		const Private::FFunctionalTestLoginUser::Result& User1Info = GetOpDataChecked<Private::FFunctionalTestLoginUser::Result>(InAsyncOp, User1KeyName);
+		const FCreateLobby::Result& CreateLobbyResult = GetOpDataChecked<FCreateLobby::Result>(InAsyncOp, CreateLobbyKeyName);
+
+		FModifyLobbyAttributes::Params Params;
+		Params.LocalUserId = User1Info.AccountInfo->UserId;
+		Params.LobbyId = CreateLobbyResult.Lobby->LobbyId;
+		Params.MutatedAttributes = {{TEXT("test_attribute_key"), TEXT("mutated_test_attribute_value")}};
+		return Params;
+	})
+	.Then(Private::ConsumeOperationStepResult<FModifyLobbyAttributes>(*Op, Private::FBindOperation<FModifyLobbyAttributes>(&LobbiesCommon, &FLobbiesCommon::ModifyLobbyAttributes)))
+	.Then([](TOnlineAsyncOp<FFunctionalTestLobbies>& InAsyncOp)
+	{
+		// Check that modification event was received.
+		const TSharedRef<FLobbyEventCapture>& EventCapture = GetOpDataChecked<TSharedRef<FLobbyEventCapture>>(InAsyncOp, LobbyEventCaptureKeyName);
+		if (EventCapture->LobbyAttributesChanged.Num() == 1)
+		{
+			const FLobbyAttributesChanged& Notification = EventCapture->LobbyAttributesChanged[0].Notification;
+			if (!Notification.ChangedAttributes.Difference(TSet<FLobbyAttributeId>{TEXT("test_attribute_key")}).IsEmpty() ||
+				!Notification.Lobby->Attributes.OrderIndependentCompareEqual(TMap<FLobbyAttributeId, FLobbyVariant>{{TEXT("test_attribute_key"), TEXT("mutated_test_attribute_value")}}))
+			{
+				InAsyncOp.SetError(Errors::Cancelled());
+			}
+		}
+		else
+		{
+			InAsyncOp.SetError(Errors::Cancelled());
+		}
+	})
+	.Then([](TOnlineAsyncOp<FFunctionalTestLobbies>& InAsyncOp)
+	{
+		// Prepare for next test check.
+		GetOpDataChecked<TSharedRef<FLobbyEventCapture>>(InAsyncOp, LobbyEventCaptureKeyName)->Empty();
+
+		const Private::FFunctionalTestLoginUser::Result& User1Info = GetOpDataChecked<Private::FFunctionalTestLoginUser::Result>(InAsyncOp, User1KeyName);
+		const FCreateLobby::Result& CreateLobbyResult = GetOpDataChecked<FCreateLobby::Result>(InAsyncOp, CreateLobbyKeyName);
+
+		FModifyLobbyMemberAttributes::Params Params;
+		Params.LocalUserId = User1Info.AccountInfo->UserId;
+		Params.LobbyId = CreateLobbyResult.Lobby->LobbyId;
+		Params.MutatedAttributes = {{TEXT("test_attribute_key"), TEXT("mutated_test_attribute_value")}};
+		return Params;
+	})
+	.Then(Private::ConsumeOperationStepResult<FModifyLobbyMemberAttributes>(*Op, Private::FBindOperation<FModifyLobbyMemberAttributes>(&LobbiesCommon, &FLobbiesCommon::ModifyLobbyMemberAttributes)))
+	.Then([](TOnlineAsyncOp<FFunctionalTestLobbies>& InAsyncOp)
+	{
+		const Private::FFunctionalTestLoginUser::Result& User1Info = GetOpDataChecked<Private::FFunctionalTestLoginUser::Result>(InAsyncOp, User1KeyName);
+
+		// Check that modification event was received.
+		const TSharedRef<FLobbyEventCapture>& EventCapture = GetOpDataChecked<TSharedRef<FLobbyEventCapture>>(InAsyncOp, LobbyEventCaptureKeyName);
+		if (EventCapture->LobbyMemberAttributesChanged.Num() == 1)
+		{
+			const FLobbyMemberAttributesChanged& Notification = EventCapture->LobbyMemberAttributesChanged[0].Notification;
+			if (const TSharedRef<const FLobbyMember>* MemberData = Notification.Lobby->Members.Find(User1Info.AccountInfo->UserId))
+			{
+				if (!Notification.ChangedAttributes.Difference(TSet<FLobbyAttributeId>{TEXT("test_attribute_key")}).IsEmpty() ||
+					!(**MemberData).Attributes.OrderIndependentCompareEqual(TMap<FLobbyAttributeId, FLobbyVariant>{{TEXT("test_attribute_key"), TEXT("mutated_test_attribute_value")}}))
+				{
+					InAsyncOp.SetError(Errors::Cancelled());
+				}
+			}
+			else
+			{
+				InAsyncOp.SetError(Errors::Cancelled());
+			}
+		}
+		else
+		{
+			InAsyncOp.SetError(Errors::Cancelled());
+		}
+	})
+	.Then([](TOnlineAsyncOp<FFunctionalTestLobbies>& InAsyncOp)
+	{
+		// Prepare for next test check.
+		GetOpDataChecked<TSharedRef<FLobbyEventCapture>>(InAsyncOp, LobbyEventCaptureKeyName)->Empty();
+
+		const Private::FFunctionalTestLoginUser::Result& User1Info = GetOpDataChecked<Private::FFunctionalTestLoginUser::Result>(InAsyncOp, User1KeyName);
+		const FCreateLobby::Result& CreateLobbyResult = GetOpDataChecked<FCreateLobby::Result>(InAsyncOp, CreateLobbyKeyName);
+
+		FLeaveLobby::Params Params;
+		Params.LocalUserId = User1Info.AccountInfo->UserId;
+		Params.LobbyId = CreateLobbyResult.Lobby->LobbyId;
+		return Params;
+	})
+	.Then(Private::ConsumeOperationStepResult<FLeaveLobby>(*Op, Private::FBindOperation<FLeaveLobby>(&LobbiesCommon, &FLobbiesCommon::LeaveLobby)))
+	.Then([](TOnlineAsyncOp<FFunctionalTestLobbies>& InAsyncOp)
+	{
+		// Check for expected events.
+		const TSharedRef<FLobbyEventCapture>& EventCapture = GetOpDataChecked<TSharedRef<FLobbyEventCapture>>(InAsyncOp, LobbyEventCaptureKeyName);
+		if (EventCapture->GetTotalNotificationsReceived() != 2
+			|| EventCapture->LobbyLeft.Num() != 1 || EventCapture->LobbyMemberLeft.Num() != 1
+			|| EventCapture->LobbyLeft[0].GlobalIndex < EventCapture->LobbyMemberLeft[0].GlobalIndex)
+		{
+			InAsyncOp.SetError(Errors::Cancelled());
+		}
+	})
+
+	//----------------------------------------------------------------------------------------------
+	// Test 2:
+	//    Step 1: Create a lobby with primary user.
+	//    Step 2: Primary user invites secondary user.
+	//    Step 3: Join lobby with secondary user.
+	//    Step 4: Leave lobby with secondary user.
+	//    Step 5: Leave lobby with primary user.
+	//----------------------------------------------------------------------------------------------
+
+	.Then([](TOnlineAsyncOp<FFunctionalTestLobbies>& InAsyncOp)
+	{
+		// Prepare for next test check.
+		GetOpDataChecked<TSharedRef<FLobbyEventCapture>>(InAsyncOp, LobbyEventCaptureKeyName)->Empty();
+
+		const Private::FFunctionalTestLoginUser::Result& User1Info = GetOpDataChecked<Private::FFunctionalTestLoginUser::Result>(InAsyncOp, User1KeyName);
+
+		FCreateLobby::Params Params;
+		Params.LocalUserId = User1Info.AccountInfo->UserId;
+		Params.LocalName = TEXT("test");
+		Params.SchemaName = TEXT("test");
+		Params.MaxMembers = 2;
+		Params.JoinPolicy = ELobbyJoinPolicy::InvitationOnly;
+		//Params.Attributes;
+		Params.LocalUsers.Emplace(FJoinLobbyLocalUserData{User1Info.AccountInfo->UserId, {}});
+		return Params;
+	})
+	.Then(Private::CaptureOperationStepResult<FCreateLobby>(*Op, CreateLobbyKeyName, Private::FBindOperation<FCreateLobby>(&LobbiesCommon, &FLobbiesCommon::CreateLobby)))
+	.Then([](TOnlineAsyncOp<FFunctionalTestLobbies>& InAsyncOp)
+	{
+		// Check both lobby joined events were received and in the correct order.
+		const TSharedRef<FLobbyEventCapture>& EventCapture = GetOpDataChecked<TSharedRef<FLobbyEventCapture>>(InAsyncOp, LobbyEventCaptureKeyName);
+		if (EventCapture->GetTotalNotificationsReceived() != 2
+			|| EventCapture->LobbyJoined.Num() != 1 || EventCapture->LobbyMemberJoined.Num() != 1
+			|| EventCapture->LobbyJoined[0].GlobalIndex > EventCapture->LobbyMemberJoined[0].GlobalIndex)
+		{
+			InAsyncOp.SetError(Errors::Cancelled());
+		}
+	})
+	.Then([](TOnlineAsyncOp<FFunctionalTestLobbies>& InAsyncOp)
+	{
+		// Prepare for next test check.
+		GetOpDataChecked<TSharedRef<FLobbyEventCapture>>(InAsyncOp, LobbyEventCaptureKeyName)->Empty();
+
+		const Private::FFunctionalTestLoginUser::Result& User1Info = GetOpDataChecked<Private::FFunctionalTestLoginUser::Result>(InAsyncOp, User1KeyName);
+		const Private::FFunctionalTestLoginUser::Result& User2Info = GetOpDataChecked<Private::FFunctionalTestLoginUser::Result>(InAsyncOp, User2KeyName);
+		const FCreateLobby::Result& CreateLobbyResult = GetOpDataChecked<FCreateLobby::Result>(InAsyncOp, CreateLobbyKeyName);
+
+		FInviteLobbyMember::Params Params;
+		Params.LocalUserId = User1Info.AccountInfo->UserId;
+		Params.LobbyId = CreateLobbyResult.Lobby->LobbyId;
+		Params.TargetUserId = User2Info.AccountInfo->UserId;
+		return Params;
+	})
+	.Then(Private::ConsumeOperationStepResult<FInviteLobbyMember>(*Op, Private::FBindOperation<FInviteLobbyMember>(&LobbiesCommon, &FLobbiesCommon::InviteLobbyMember)))
+	.Then([TestConfig, LobbyEventsPtr = &LobbyEvents](TOnlineAsyncOp<FFunctionalTestLobbies>& InAsyncOp)
+	{
+		const Private::FFunctionalTestLoginUser::Result& User2Info = GetOpDataChecked<Private::FFunctionalTestLoginUser::Result>(InAsyncOp, User2KeyName);
+		const FCreateLobby::Result& CreateLobbyResult = GetOpDataChecked<FCreateLobby::Result>(InAsyncOp, CreateLobbyKeyName);
+		return Private::AwaitInvitation(*LobbyEventsPtr, User2Info.AccountInfo->UserId, CreateLobbyResult.Lobby->LobbyId, TestConfig->InvitationWaitSeconds);
+	})
+	.Then([](TOnlineAsyncOp<FFunctionalTestLobbies>& InAsyncOp, TDefaultErrorResultInternal<FOnlineLobbyIdHandle>&& Result)
+	{
+		if (Result.IsError())
+		{
+			InAsyncOp.SetError(Errors::Cancelled(Result.GetErrorValue()));
+		}
+	})
+	.Then([](TOnlineAsyncOp<FFunctionalTestLobbies>& InAsyncOp)
+	{
+		// Prepare for next test check.
+		GetOpDataChecked<TSharedRef<FLobbyEventCapture>>(InAsyncOp, LobbyEventCaptureKeyName)->Empty();
+
+		const Private::FFunctionalTestLoginUser::Result& User2Info = GetOpDataChecked<Private::FFunctionalTestLoginUser::Result>(InAsyncOp, User2KeyName);
+		const FCreateLobby::Result& CreateLobbyResult = GetOpDataChecked<FCreateLobby::Result>(InAsyncOp, CreateLobbyKeyName);
+
+		FJoinLobby::Params Params;
+		Params.LocalUserId = User2Info.AccountInfo->UserId;
+		Params.LocalName = TEXT("test");
+		Params.LobbyId = CreateLobbyResult.Lobby->LobbyId;
+		Params.LocalUsers.Emplace(FJoinLobbyLocalUserData{User2Info.AccountInfo->UserId, {}});
+		return Params;
+	})
+	.Then(Private::ConsumeOperationStepResult<FJoinLobby>(*Op, Private::FBindOperation<FJoinLobby>(&LobbiesCommon, &FLobbiesCommon::JoinLobby)))
+	.Then([](TOnlineAsyncOp<FFunctionalTestLobbies>& InAsyncOp)
+	{
+		// Check both lobby joined events were received and in the correct order.
+		const TSharedRef<FLobbyEventCapture>& EventCapture = GetOpDataChecked<TSharedRef<FLobbyEventCapture>>(InAsyncOp, LobbyEventCaptureKeyName);
+		if (EventCapture->GetTotalNotificationsReceived() != 1 || EventCapture->LobbyMemberJoined.Num() != 1)
+		{
+			InAsyncOp.SetError(Errors::Cancelled());
+		}
+	})
+	.Then([](TOnlineAsyncOp<FFunctionalTestLobbies>& InAsyncOp)
+	{
+		// Prepare for next test check.
+		GetOpDataChecked<TSharedRef<FLobbyEventCapture>>(InAsyncOp, LobbyEventCaptureKeyName)->Empty();
+
+		const Private::FFunctionalTestLoginUser::Result& User2Info = GetOpDataChecked<Private::FFunctionalTestLoginUser::Result>(InAsyncOp, User2KeyName);
+		const FCreateLobby::Result& CreateLobbyResult = GetOpDataChecked<FCreateLobby::Result>(InAsyncOp, CreateLobbyKeyName);
+
+		FLeaveLobby::Params Params;
+		Params.LocalUserId = User2Info.AccountInfo->UserId;
+		Params.LobbyId = CreateLobbyResult.Lobby->LobbyId;
+		return Params;
+	})
+	.Then(Private::ConsumeOperationStepResult<FLeaveLobby>(*Op, Private::FBindOperation<FLeaveLobby>(&LobbiesCommon, &FLobbiesCommon::LeaveLobby)))
+	.Then([](TOnlineAsyncOp<FFunctionalTestLobbies>& InAsyncOp)
+	{
+		// Check for expected events.
+		const TSharedRef<FLobbyEventCapture>& EventCapture = GetOpDataChecked<TSharedRef<FLobbyEventCapture>>(InAsyncOp, LobbyEventCaptureKeyName);
+		if (EventCapture->GetTotalNotificationsReceived() != 1 || EventCapture->LobbyMemberLeft.Num() != 1)
+		{
+			InAsyncOp.SetError(Errors::Cancelled());
+		}
+	})
+	.Then([](TOnlineAsyncOp<FFunctionalTestLobbies>& InAsyncOp)
+	{
+		// Prepare for next test check.
+		GetOpDataChecked<TSharedRef<FLobbyEventCapture>>(InAsyncOp, LobbyEventCaptureKeyName)->Empty();
+
+		const Private::FFunctionalTestLoginUser::Result& User1Info = GetOpDataChecked<Private::FFunctionalTestLoginUser::Result>(InAsyncOp, User1KeyName);
+		const FCreateLobby::Result& CreateLobbyResult = GetOpDataChecked<FCreateLobby::Result>(InAsyncOp, CreateLobbyKeyName);
+
+		FLeaveLobby::Params Params;
+		Params.LocalUserId = User1Info.AccountInfo->UserId;
+		Params.LobbyId = CreateLobbyResult.Lobby->LobbyId;
+		return Params;
+	})
+	.Then(Private::ConsumeOperationStepResult<FLeaveLobby>(*Op, Private::FBindOperation<FLeaveLobby>(&LobbiesCommon, &FLobbiesCommon::LeaveLobby)))
+	.Then([](TOnlineAsyncOp<FFunctionalTestLobbies>& InAsyncOp)
+	{
+		// Check for expected events.
+		const TSharedRef<FLobbyEventCapture>& EventCapture = GetOpDataChecked<TSharedRef<FLobbyEventCapture>>(InAsyncOp, LobbyEventCaptureKeyName);
+		if (EventCapture->GetTotalNotificationsReceived() != 2
+			|| EventCapture->LobbyLeft.Num() != 1 || EventCapture->LobbyMemberLeft.Num() != 1
+			|| EventCapture->LobbyLeft[0].GlobalIndex < EventCapture->LobbyMemberLeft[0].GlobalIndex)
+		{
+			InAsyncOp.SetError(Errors::Cancelled());
+		}
+	})
+
+	//----------------------------------------------------------------------------------------------
+	// Test 3: Simple search
+	//    Step 1: Create a lobby with primary user.
+	//    Step 2: Secondary user searches for lobbies.
+	//    Step 3: Join lobby with secondary user.
+	//    Step 4: Leave lobby with secondary user.
+	//    Step 5: Leave lobby with primary user.
+	//----------------------------------------------------------------------------------------------
+
+	.Then([](TOnlineAsyncOp<FFunctionalTestLobbies>& InAsyncOp)
+	{
+		// Prepare for next test check.
+		GetOpDataChecked<TSharedRef<FLobbyEventCapture>>(InAsyncOp, LobbyEventCaptureKeyName)->Empty();
+
+		const Private::FFunctionalTestLoginUser::Result& User1Info = GetOpDataChecked<Private::FFunctionalTestLoginUser::Result>(InAsyncOp, User1KeyName);
+
+		FCreateLobby::Params Params;
+		Params.LocalUserId = User1Info.AccountInfo->UserId;
+		Params.LocalName = TEXT("test");
+		Params.SchemaName = TEXT("test");
+		Params.MaxMembers = 2;
+		Params.JoinPolicy = ELobbyJoinPolicy::PublicAdvertised;
+		//Params.Attributes;
+		Params.LocalUsers.Emplace(FJoinLobbyLocalUserData{User1Info.AccountInfo->UserId, {}});
+		return Params;
+	})
+	.Then(Private::CaptureOperationStepResult<FCreateLobby>(*Op, CreateLobbyKeyName, Private::FBindOperation<FCreateLobby>(&LobbiesCommon, &FLobbiesCommon::CreateLobby)))
+	.Then([](TOnlineAsyncOp<FFunctionalTestLobbies>& InAsyncOp)
+	{
+		// Check both lobby joined events were received and in the correct order.
+		const TSharedRef<FLobbyEventCapture>& EventCapture = GetOpDataChecked<TSharedRef<FLobbyEventCapture>>(InAsyncOp, LobbyEventCaptureKeyName);
+		if (EventCapture->GetTotalNotificationsReceived() != 2
+			|| EventCapture->LobbyJoined.Num() != 1 || EventCapture->LobbyMemberJoined.Num() != 1
+			|| EventCapture->LobbyJoined[0].GlobalIndex > EventCapture->LobbyMemberJoined[0].GlobalIndex)
+		{
+			InAsyncOp.SetError(Errors::Cancelled());
+		}
+	})
+	.Then([](TOnlineAsyncOp<FFunctionalTestLobbies>& InAsyncOp)
+	{
+		const Private::FFunctionalTestLoginUser::Result& User2Info = GetOpDataChecked<Private::FFunctionalTestLoginUser::Result>(InAsyncOp, User2KeyName);
+		const FCreateLobby::Result& CreateLobbyResult = GetOpDataChecked<FCreateLobby::Result>(InAsyncOp, CreateLobbyKeyName);
+
+		// Search for lobby from create
+		FFindLobbies::Params Params;
+		Params.LocalUserId = User2Info.AccountInfo->UserId;
+		Params.LobbyId = CreateLobbyResult.Lobby->LobbyId;
+		return Params;
+	})
+	.Then(Private::CaptureOperationStepResult<FFindLobbies>(*Op, FindLobbyKeyName, Private::FBindOperation<FFindLobbies>(&LobbiesCommon, &FLobbiesCommon::FindLobbies)))
+	.Then([](TOnlineAsyncOp<FFunctionalTestLobbies>& InAsyncOp)
+	{
+		// Prepare for next test check.
+		GetOpDataChecked<TSharedRef<FLobbyEventCapture>>(InAsyncOp, LobbyEventCaptureKeyName)->Empty();
+
+		const Private::FFunctionalTestLoginUser::Result& User2Info = GetOpDataChecked<Private::FFunctionalTestLoginUser::Result>(InAsyncOp, User2KeyName);
+		const FCreateLobby::Result& CreateLobbyResult = GetOpDataChecked<FCreateLobby::Result>(InAsyncOp, CreateLobbyKeyName);
+		const FFindLobbies::Result& FindResults = GetOpDataChecked<FFindLobbies::Result>(InAsyncOp, FindLobbyKeyName);
+
+		if (FindResults.Lobbies.Num() != 1)
+		{
+			InAsyncOp.SetError(Errors::Cancelled());
+			return FJoinLobby::Params{};
+		}
+
+		FJoinLobby::Params Params;
+		Params.LocalUserId = User2Info.AccountInfo->UserId;
+		Params.LocalName = TEXT("test");
+		Params.LobbyId = FindResults.Lobbies[0]->LobbyId;
+		Params.LocalUsers.Emplace(FJoinLobbyLocalUserData{User2Info.AccountInfo->UserId, {}});
+		return Params;
+	})
+	.Then(Private::ConsumeOperationStepResult<FJoinLobby>(*Op, Private::FBindOperation<FJoinLobby>(&LobbiesCommon, &FLobbiesCommon::JoinLobby)))
+	.Then([](TOnlineAsyncOp<FFunctionalTestLobbies>& InAsyncOp)
+	{
+		// Check both lobby joined events were received and in the correct order.
+		const TSharedRef<FLobbyEventCapture>& EventCapture = GetOpDataChecked<TSharedRef<FLobbyEventCapture>>(InAsyncOp, LobbyEventCaptureKeyName);
+		if (EventCapture->GetTotalNotificationsReceived() != 1 || EventCapture->LobbyMemberJoined.Num() != 1)
+		{
+			InAsyncOp.SetError(Errors::Cancelled());
+		}
+	})
+	.Then([](TOnlineAsyncOp<FFunctionalTestLobbies>& InAsyncOp)
+	{
+		// Prepare for next test check.
+		GetOpDataChecked<TSharedRef<FLobbyEventCapture>>(InAsyncOp, LobbyEventCaptureKeyName)->Empty();
+
+		const Private::FFunctionalTestLoginUser::Result& User2Info = GetOpDataChecked<Private::FFunctionalTestLoginUser::Result>(InAsyncOp, User2KeyName);
+		const FCreateLobby::Result& CreateLobbyResult = GetOpDataChecked<FCreateLobby::Result>(InAsyncOp, CreateLobbyKeyName);
+
+		FLeaveLobby::Params Params;
+		Params.LocalUserId = User2Info.AccountInfo->UserId;
+		Params.LobbyId = CreateLobbyResult.Lobby->LobbyId;
+		return Params;
+	})
+	.Then(Private::ConsumeOperationStepResult<FLeaveLobby>(*Op, Private::FBindOperation<FLeaveLobby>(&LobbiesCommon, &FLobbiesCommon::LeaveLobby)))
+	.Then([](TOnlineAsyncOp<FFunctionalTestLobbies>& InAsyncOp)
+	{
+		// Check for expected events.
+		const TSharedRef<FLobbyEventCapture>& EventCapture = GetOpDataChecked<TSharedRef<FLobbyEventCapture>>(InAsyncOp, LobbyEventCaptureKeyName);
+		if (EventCapture->GetTotalNotificationsReceived() != 1|| EventCapture->LobbyMemberLeft.Num() != 1)
+		{
+			InAsyncOp.SetError(Errors::Cancelled());
+		}
+	})
+	.Then([](TOnlineAsyncOp<FFunctionalTestLobbies>& InAsyncOp)
+	{
+		// Prepare for next test check.
+		GetOpDataChecked<TSharedRef<FLobbyEventCapture>>(InAsyncOp, LobbyEventCaptureKeyName)->Empty();
+
+		const Private::FFunctionalTestLoginUser::Result& User1Info = GetOpDataChecked<Private::FFunctionalTestLoginUser::Result>(InAsyncOp, User1KeyName);
+		const FCreateLobby::Result& CreateLobbyResult = GetOpDataChecked<FCreateLobby::Result>(InAsyncOp, CreateLobbyKeyName);
+
+		FLeaveLobby::Params Params;
+		Params.LocalUserId = User1Info.AccountInfo->UserId;
+		Params.LobbyId = CreateLobbyResult.Lobby->LobbyId;
+		return Params;
+	})
+	.Then(Private::ConsumeOperationStepResult<FLeaveLobby>(*Op, Private::FBindOperation<FLeaveLobby>(&LobbiesCommon, &FLobbiesCommon::LeaveLobby)))
+	.Then([](TOnlineAsyncOp<FFunctionalTestLobbies>& InAsyncOp)
+	{
+		// Check for expected events.
+		const TSharedRef<FLobbyEventCapture>& EventCapture = GetOpDataChecked<TSharedRef<FLobbyEventCapture>>(InAsyncOp, LobbyEventCaptureKeyName);
+		if (EventCapture->GetTotalNotificationsReceived() != 2
+			|| EventCapture->LobbyLeft.Num() != 1 || EventCapture->LobbyMemberLeft.Num() != 1
+			|| EventCapture->LobbyLeft[0].GlobalIndex < EventCapture->LobbyMemberLeft[0].GlobalIndex)
+		{
+			InAsyncOp.SetError(Errors::Cancelled());
+		}
+	})
+
+	//----------------------------------------------------------------------------------------------
+	// Complete
+	//----------------------------------------------------------------------------------------------
+
+	.Then([](TOnlineAsyncOp<FFunctionalTestLobbies>& InAsyncOp)
+	{
+		InAsyncOp.SetResult(FFunctionalTestLobbies::Result{});
+	})
+	.Enqueue(LobbiesCommon.GetServices().GetParallelQueue());
+
+	return Op->GetHandle();
+}
+
+#endif // LOBBIES_FUNCTIONAL_TEST_ENABLED
+
+/* UE::Online */ }
diff --git a/Engine/Plugins/Online/Experimental/OnlineServices/Source/OnlineServicesCommon/Private/Online/LobbiesCommonTests.h b/Engine/Plugins/Online/Experimental/OnlineServices/Source/OnlineServicesCommon/Private/Online/LobbiesCommonTests.h
new file mode 100644
--- /dev/null
+++ b/Engine/Plugins/Online/Experimental/OnlineServices/Source/OnlineServicesCommon/Private/Online/LobbiesCommonTests.h
@@ -0,0 +1,21 @@
+// Copyright Epic Games, Inc. All Rights Reserved.
+
+#pragma once
+
+#include "Online/LobbiesCommonTypes.h"
+#include "Online/OnlineAsyncOpHandle.h"
+#include "Online/OnlineMeta.h"
+
+namespace UE::Online {
+
+class IAuth;
+class FLobbiesCommon;
+struct FLobbyEvents;
+
+#if LOBBIES_FUNCTIONAL_TEST_ENABLED
+
+TOnlineAsyncOpHandle<FFunctionalTestLobbies> RunLobbyFunctionalTest(IAuth& AuthInterface, FLobbiesCommon& LobbiesCommon, FLobbyEvents& LobbyEvents);
+
+#endif // LOBBIES_FUNCTIONAL_TEST_ENABLED
+
+/* UE::Online */ }
\ No newline at end of file
diff --git a/Engine/Plugins/Online/Experimental/OnlineServices/Source/OnlineServicesCommon/Private/Online/LobbiesCommonTypes.cpp b/Engine/Plugins/Online/Experimental/OnlineServices/Source/OnlineServicesCommon/Private/Online/LobbiesCommonTypes.cpp
new file mode 100644
--- /dev/null
+++ b/Engine/Plugins/Online/Experimental/OnlineServices/Source/OnlineServicesCommon/Private/Online/LobbiesCommonTypes.cpp
@@ -0,0 +1,520 @@
+// Copyright Epic Games, Inc. All Rights Reserved.
+
+#include "Online/LobbiesCommonTypes.h"
+
+namespace UE::Online {
+
+TSharedPtr<FLobbySchema> FLobbySchema::Create(FLobbySchemaConfig LobbySchemaConfig)
+{
+	TSharedRef<FLobbySchema> Schema = MakeShared<FLobbySchema>();
+	Schema->SchemaId = LobbySchemaConfig.SchemaName;
+
+	// Todo: for real validation and init
+	if (Schema->SchemaId == FLobbySchemaId())
+	{
+		return TSharedPtr<FLobbySchema>();
+	}
+	else
+	{
+		return Schema;
+	}
+}
+
+TSharedRef<FLobbyServiceAttributeChanges> FLobbySchema::TranslateLobbyAttributes(const TMap<FLobbyAttributeId, FLobbyVariant>& LobbyAttributes, const FLobbyClientAttributeChanges& ClientAttributeChanges) const
+{
+	// Todo: implementation.
+
+	TSharedRef<FLobbyServiceAttributeChanges> Changes = MakeShared<FLobbyServiceAttributeChanges>();
+	Changes->MutatedAttributes = ClientAttributeChanges.MutatedAttributes;
+	Changes->ClearedAttributes = ClientAttributeChanges.ClearedAttributes;
+	return Changes;
+}
+
+TSharedRef<FLobbyClientAttributeChanges> FLobbySchema::TranslateLobbyAttributes(const FLobbyServiceAttributeChanges& ServiceAttributeChanges) const
+{
+	// Todo: implementation.
+
+	TSharedRef<FLobbyClientAttributeChanges> Changes = MakeShared<FLobbyClientAttributeChanges>();
+	Changes->MutatedAttributes = ServiceAttributeChanges.MutatedAttributes;
+	Changes->ClearedAttributes = ServiceAttributeChanges.ClearedAttributes;
+	return Changes;
+}
+
+TSharedRef<FLobbyServiceAttributeChanges> FLobbySchema::TranslateLobbyMemberAttributes(const TMap<FLobbyAttributeId, FLobbyVariant>& LobbyAttributes, const FLobbyClientAttributeChanges& ClientAttributeChanges) const
+{
+	// Todo: implementation.
+
+	TSharedRef<FLobbyServiceAttributeChanges> Changes = MakeShared<FLobbyServiceAttributeChanges>();
+	Changes->MutatedAttributes = ClientAttributeChanges.MutatedAttributes;
+	Changes->ClearedAttributes = ClientAttributeChanges.ClearedAttributes;
+	return Changes;
+}
+
+TSharedRef<FLobbyClientAttributeChanges> FLobbySchema::TranslateLobbyMemberAttributes(const FLobbyServiceAttributeChanges& ServiceAttributeChanges) const
+{
+	// Todo: implementation.
+
+	TSharedRef<FLobbyClientAttributeChanges> Changes = MakeShared<FLobbyClientAttributeChanges>();
+	Changes->MutatedAttributes = ServiceAttributeChanges.MutatedAttributes;
+	Changes->ClearedAttributes = ServiceAttributeChanges.ClearedAttributes;
+	return Changes;
+}
+
+bool FLobbySchemaRegistry::Initialize(TArray<FLobbySchemaConfig> LobbySchemaConfigs)
+{
+	bool bSuccess = true;
+
+	for (FLobbySchemaConfig& LobbySchemaConfig : LobbySchemaConfigs)
+	{
+		bSuccess &= RegisterSchema(MoveTemp(LobbySchemaConfig));
+	}
+
+	return bSuccess;
+}
+
+TSharedPtr<FLobbySchema> FLobbySchemaRegistry::FindSchema(FLobbySchemaId SchemaId)
+{
+	TSharedRef<FLobbySchema>* Schema = RegisteredSchemas.Find(SchemaId);
+	return Schema ? *Schema : TSharedPtr<FLobbySchema>();
+}
+
+bool FLobbySchemaRegistry::RegisterSchema(FLobbySchemaConfig LobbySchemaConfig)
+{
+	// Todo: handle schema hierarchy.
+	// Todo: handle schema parsing.
+
+	FLobbySchemaId SchemaName = LobbySchemaConfig.SchemaName;
+
+	if (TSharedPtr<FLobbySchema> Schema = FLobbySchema::Create(MoveTemp(LobbySchemaConfig)))
+	{
+		RegisteredSchemas.Emplace(SchemaName, Schema.ToSharedRef());
+		return true;
+	}
+	else
+	{
+		return false;
+	}
+}
+
+
+FClientLobbyData::FClientLobbyData(FOnlineLobbyIdHandle LobbyId)
+	: PublicData(MakeShared<FLobby>())
+{
+	PublicData->LobbyId = LobbyId;
+}
+
+TSharedPtr<const FClientLobbyMemberSnapshot> FClientLobbyData::GetMemberData(FOnlineAccountIdHandle MemberAccountId) const
+{
+	const TSharedRef<FClientLobbyMemberSnapshot>* FoundMemberData = MemberDataStorage.Find(MemberAccountId);
+	return FoundMemberData ? *FoundMemberData : TSharedPtr<const FClientLobbyMemberSnapshot>();
+}
+
+FApplyLobbyUpdateResult FClientLobbyData::ApplyLobbyUpdateFromServiceSnapshot(
+	FClientLobbySnapshot&& LobbySnapshot,
+	TMap<FOnlineAccountIdHandle, TSharedRef<FClientLobbyMemberSnapshot>>&& LobbyMemberSnapshots,
+	TMap<FOnlineAccountIdHandle, ELobbyMemberLeaveReason>&& LeaveReasons,
+	FLobbyEvents* LobbyEvents)
+{
+	FApplyLobbyUpdateResult Result;
+	const bool bDispatchNotifications = LobbyEvents != nullptr && !LocalMembers.IsEmpty();
+
+	// Notify if schema changed.
+	// Schema change notification must be processed first as the schema affects how the client processes attributes.
+	if (LobbySnapshot.SchemaName != PublicData->SchemaName)
+	{
+		PublicData->SchemaName = LobbySnapshot.SchemaName;
+		if (bDispatchNotifications)
+		{
+			LobbyEvents->OnLobbySchemaChanged.Broadcast(FLobbySchemaChanged{PublicData});
+		}
+	}
+
+	// Handle lobby attribute changes.
+	TSet<FLobbyAttributeId> LobbyAttributeChanges = ApplyAttributeUpdateFromSnapshot(MoveTemp(LobbySnapshot.Attributes), PublicData->Attributes);
+	if (bDispatchNotifications && !LobbyAttributeChanges.IsEmpty())
+	{
+		LobbyEvents->OnLobbyAttributesChanged.Broadcast(FLobbyAttributesChanged{PublicData, MoveTemp(LobbyAttributeChanges)});
+	}
+
+	// Process members who left.
+	{
+		// Build list of leaving members.
+		TArray<TPair<TSharedRef<FClientLobbyMemberSnapshot>, ELobbyMemberLeaveReason>> LeavingMembers;
+		for (TPair<FOnlineAccountIdHandle, TSharedRef<FClientLobbyMemberSnapshot>>& MemberData : MemberDataStorage)
+		{
+			// Check if member exists in new snapshot.
+			if (LobbySnapshot.Members.Find(MemberData.Key) == nullptr)
+			{
+				// Member has left - check if they gave an explicit reason.
+				ELobbyMemberLeaveReason ResovedReason = ELobbyMemberLeaveReason::Disconnected;
+				if (const ELobbyMemberLeaveReason* MemberLeaveReason = LeaveReasons.Find(MemberData.Key))
+				{
+					ResovedReason = *MemberLeaveReason;
+				}
+				else
+				{
+					UE_LOG(LogTemp, Warning, TEXT("[FClientLobbyData::ApplyLobbyUpdateFromSnapshot] Member left lobby without giving reason: Lobby: %s, Member: %s"),
+						*ToLogString(PublicData->LobbyId), *ToLogString(MemberData.Key));
+				}
+
+				LeavingMembers.Emplace(MemberData.Value, ResovedReason);
+			}
+		}
+
+		// Remove member data and notify.
+		for (const TPair<TSharedRef<FClientLobbyMemberSnapshot>, ELobbyMemberLeaveReason>& MemberData : LeavingMembers)
+		{
+			MemberDataStorage.Remove(MemberData.Key->AccountId);
+			PublicData->Members.Remove(MemberData.Key->AccountId);
+
+			if (MemberData.Key->bIsLocalMember)
+			{
+				LocalMembers.Remove(MemberData.Key->AccountId);
+				Result.LeavingLocalMembers.Add(MemberData.Key->AccountId);
+			}
+
+			if (bDispatchNotifications)
+			{
+				LobbyEvents->OnLobbyMemberLeft.Broadcast(FLobbyMemberLeft{PublicData, MemberData.Key, MemberData.Value});
+			}
+		}
+	}
+
+	// Process member joins and attribute changes.
+	// Joining members are expected to have a member data snapshot.
+	for (TPair<FOnlineAccountIdHandle, TSharedRef<FClientLobbyMemberSnapshot>>& MemberSnapshot : LobbyMemberSnapshots)
+	{
+		if (LobbySnapshot.Members.Find(MemberSnapshot.Key) == nullptr)
+		{
+			UE_LOG(LogTemp, Warning, TEXT("[FClientLobbyData::ApplyLobbyUpdateFromSnapshot] Member update ignored for unknown member: Lobby: %s, Member: %s"),
+				*ToLogString(PublicData->LobbyId), *ToLogString(MemberSnapshot.Key));
+		}
+		else
+		{
+			// Check whether this member is already in the lobby.
+			if (TSharedRef<FClientLobbyMemberSnapshot>* MemberData = MemberDataStorage.Find(MemberSnapshot.Key))
+			{
+				// Member is already in the lobby. Check if their attributes changed.
+				TSet<FLobbyAttributeId> MemberAttributeChanges = ApplyAttributeUpdateFromSnapshot(MoveTemp(MemberSnapshot.Value->Attributes), (*MemberData)->Attributes);
+				if (bDispatchNotifications && !MemberAttributeChanges.IsEmpty())
+				{
+					LobbyEvents->OnLobbyMemberAttributesChanged.Broadcast(FLobbyMemberAttributesChanged{PublicData, (*MemberData), MemberAttributeChanges});
+				}
+			}
+			else
+			{
+				// Member is not in the lobby, add them.
+				MemberDataStorage.Add(MemberSnapshot.Key, MemberSnapshot.Value);
+				PublicData->Members.Add(MemberSnapshot.Key, MoveTemp(MemberSnapshot.Value));
+				if (bDispatchNotifications)
+				{
+					LobbyEvents->OnLobbyMemberJoined.Broadcast(FLobbyMemberJoined{PublicData, MemberSnapshot.Value});
+				}
+			}
+		}
+	}
+
+	// Process lobby state.
+	// The lobby owner notification must be processed after member updates as the new owner may not
+	// have existed in the previous snapshot.
+	if (PublicData->OwnerAccountId != LobbySnapshot.OwnerAccountId)
+	{
+		if (TSharedRef<FClientLobbyMemberSnapshot>* MemberData = MemberDataStorage.Find(LobbySnapshot.OwnerAccountId))
+		{
+			PublicData->OwnerAccountId = LobbySnapshot.OwnerAccountId;
+			if (bDispatchNotifications)
+			{
+				LobbyEvents->OnLobbyLeaderChanged.Broadcast(FLobbyLeaderChanged{PublicData, (*MemberData)});
+			}
+		}
+		else
+		{
+			UE_LOG(LogTemp, Warning, TEXT("[FClientLobbyData::ApplyLobbyUpdateFromSnapshot] Lobby owner chage failed - lobby member data not found: Lobby: %s, Member: %s"),
+				*ToLogString(PublicData->LobbyId), *ToLogString(LobbySnapshot.OwnerAccountId));
+		}
+	}
+
+	// Set remaining lobby data.
+	PublicData->MaxMembers = LobbySnapshot.MaxMembers;
+	PublicData->JoinPolicy = LobbySnapshot.JoinPolicy;
+
+	// If no local members remain in the lobby process lobby removal.
+	if (bDispatchNotifications && LocalMembers.IsEmpty())
+	{
+		TArray<TSharedRef<FClientLobbyMemberSnapshot>> LeavingMembers;
+		MemberDataStorage.GenerateValueArray(LeavingMembers);
+
+		for (TSharedRef<FClientLobbyMemberSnapshot>& LeavingMember : LeavingMembers)
+		{
+			MemberDataStorage.Remove(LeavingMember->AccountId);
+			PublicData->Members.Remove(LeavingMember->AccountId);
+			LobbyEvents->OnLobbyMemberLeft.Broadcast(FLobbyMemberLeft{PublicData, LeavingMember, ELobbyMemberLeaveReason::Left});
+		}
+
+		LobbyEvents->OnLobbyLeft.Broadcast(FLobbyLeft{PublicData});
+	}
+
+	return Result;
+}
+
+FApplyLobbyUpdateResult FClientLobbyData::ApplyLobbyUpdateFromLocalChanges(FClientLobbyDataChanges&& Changes, FLobbyEvents& LobbyEvents)
+{
+	FApplyLobbyUpdateResult Result;
+	const bool bDispatchLobbyNotifications = !LocalMembers.IsEmpty();
+
+	// Process lobby changes.
+	{
+		if (Changes.LocalName)
+		{
+			PublicData->LocalName = *Changes.LocalName;
+		}
+
+		// Join policy.
+		if (Changes.JoinPolicy)
+		{
+			PublicData->JoinPolicy = *Changes.JoinPolicy;
+		}
+
+		// Lobby owner.
+		if (Changes.OwnerAccountId && *Changes.OwnerAccountId != PublicData->OwnerAccountId)
+		{
+			if (TSharedRef<FClientLobbyMemberSnapshot>* MemberData = MemberDataStorage.Find(*Changes.OwnerAccountId))
+			{
+				PublicData->OwnerAccountId = *Changes.OwnerAccountId;
+				if (bDispatchLobbyNotifications)
+				{
+					LobbyEvents.OnLobbyLeaderChanged.Broadcast(FLobbyLeaderChanged{PublicData, (*MemberData)});
+				}
+			}
+			else
+			{
+				UE_LOG(LogTemp, Warning, TEXT("[FClientLobbyData::ApplyLobbyUpdateFromLocalChanges] Lobby owner chage failed - lobby member data not found: Lobby: %s, Member: %s"),
+					*ToLogString(PublicData->LobbyId), *ToLogString(*Changes.OwnerAccountId));
+			}
+		}
+
+		// Lobby schema.
+		if (Changes.LobbySchema && *Changes.LobbySchema != PublicData->SchemaName)
+		{
+			PublicData->SchemaName = *Changes.LobbySchema;
+			if (bDispatchLobbyNotifications)
+			{
+				LobbyEvents.OnLobbySchemaChanged.Broadcast(FLobbySchemaChanged{PublicData});
+			}
+		}
+
+		// Apply lobby attribute changes.
+		if (!Changes.MutatedAttributes.IsEmpty() || !Changes.ClearedAttributes.IsEmpty())
+		{
+			TSet<FLobbyAttributeId> ChangedLobbyAttributes = ApplyAttributeUpdateFromChanges(
+				MoveTemp(Changes.MutatedAttributes),
+				MoveTemp(Changes.ClearedAttributes),
+				PublicData->Attributes);
+
+			if (bDispatchLobbyNotifications && !ChangedLobbyAttributes.IsEmpty())
+			{
+				LobbyEvents.OnLobbyAttributesChanged.Broadcast(FLobbyAttributesChanged{PublicData, ChangedLobbyAttributes});
+			}
+		}
+	}
+
+	// Apply member changes.
+	{
+		// Adding local members will make the lobby active. When the lobby moves from inactive to
+		// active join events must be generated for the lobby and each of its existing members.
+		if (!Changes.MutatedMembers.IsEmpty())
+		{
+			// When adding members for the first time, notify for lobby joined then notify lobby member joined for each existing member.
+			if (LocalMembers.IsEmpty())
+			{
+				LobbyEvents.OnLobbyJoined.Broadcast(FLobbyJoined{PublicData});
+
+				for (TPair<FOnlineAccountIdHandle, TSharedRef<FClientLobbyMemberSnapshot>>& MemberData : MemberDataStorage)
+				{
+					// The lobby creator will already be in the list during lobby creation due to applying a lobby snapshot.
+					if (Changes.MutatedMembers.Find(MemberData.Key) == nullptr)
+					{
+						LobbyEvents.OnLobbyMemberJoined.Broadcast(FLobbyMemberJoined{PublicData, MemberData.Value});
+					}
+				}
+			}
+
+			// Create or modify state for each mutated member and notify.
+			// The lobby creator will already be in the list during lobby creation due to applying a lobby snapshot.
+			for (TPair<FOnlineAccountIdHandle, TSharedRef<FClientLobbyMemberDataChanges>>& MutatedMemberPair : Changes.MutatedMembers)
+			{
+				if (TSharedRef<FClientLobbyMemberSnapshot>* MemberDataPtr = MemberDataStorage.Find(MutatedMemberPair.Key))
+				{
+					const bool bIsJoinedLocalMember = LocalMembers.Find(MutatedMemberPair.Key) != nullptr;
+
+					TSharedRef<FClientLobbyMemberSnapshot>& MemberData = *MemberDataPtr;
+					TSet<FLobbyAttributeId> ChangedAttributes = ApplyAttributeUpdateFromChanges(
+						MoveTemp(MutatedMemberPair.Value->MutatedAttributes),
+						MoveTemp(MutatedMemberPair.Value->ClearedAttributes),
+						MemberData->Attributes);
+
+					if (bIsJoinedLocalMember && !ChangedAttributes.IsEmpty())
+					{
+						LobbyEvents.OnLobbyMemberAttributesChanged.Broadcast(FLobbyMemberAttributesChanged{PublicData, MemberData, ChangedAttributes});
+					}
+
+					if (!bIsJoinedLocalMember)
+					{
+						MemberData->bIsLocalMember = true;
+						LocalMembers.Add(MemberData->AccountId);
+						LobbyEvents.OnLobbyMemberJoined.Broadcast(FLobbyMemberJoined{PublicData, MemberData});
+					}
+				}
+				else
+				{
+					TSharedRef<FClientLobbyMemberSnapshot> JoiningMember = MakeShared<FClientLobbyMemberSnapshot>();
+					JoiningMember->bIsLocalMember = true;
+					JoiningMember->AccountId = MutatedMemberPair.Key;
+					JoiningMember->Attributes = MoveTemp(MutatedMemberPair.Value->MutatedAttributes);
+
+					// todo:
+					//NewMember->PlatformAccountId;
+					//NewMember->PlatformDisplayName;
+
+					MemberDataStorage.Add(JoiningMember->AccountId, JoiningMember);
+					PublicData->Members.Add(JoiningMember->AccountId, JoiningMember);
+					LocalMembers.Add(JoiningMember->AccountId);
+					LobbyEvents.OnLobbyMemberJoined.Broadcast(FLobbyMemberJoined{PublicData, JoiningMember});
+				}
+			}
+		}
+
+		// Handle leaving members.
+		// When the last local member leaves, leave events will be generated for each of the
+		// remaining lobby members and for the lobby itself.
+		if (!Changes.LeavingMembers.IsEmpty())
+		{
+			for (TPair<FOnlineAccountIdHandle, ELobbyMemberLeaveReason>& LeavingMemberPair : Changes.LeavingMembers)
+			{
+				// Member data is expected to be added through changes before being mutated.
+				TSharedRef<FClientLobbyMemberSnapshot>* MemberDataPtr = MemberDataStorage.Find(LeavingMemberPair.Key);
+				if (ensure(MemberDataPtr))
+				{
+					TSharedRef<FClientLobbyMemberSnapshot> MemberData = *MemberDataPtr;
+					if (MemberData->bIsLocalMember)
+					{
+						Result.LeavingLocalMembers.Add(LeavingMemberPair.Key);
+						LocalMembers.Remove(LeavingMemberPair.Key);
+					}
+
+					MemberDataStorage.Remove(LeavingMemberPair.Key);
+					PublicData->Members.Remove(LeavingMemberPair.Key);
+					LobbyEvents.OnLobbyMemberLeft.Broadcast(FLobbyMemberLeft{PublicData, MemberData, LeavingMemberPair.Value});
+				}
+			}
+
+			// If no local members remain in the lobby process lobby removal.
+			if (LocalMembers.IsEmpty())
+			{
+				TArray<TSharedRef<FClientLobbyMemberSnapshot>> LeavingMembers;
+				MemberDataStorage.GenerateValueArray(LeavingMembers);
+
+				for (TSharedRef<FClientLobbyMemberSnapshot>& LeavingMember : LeavingMembers)
+				{
+					MemberDataStorage.Remove(LeavingMember->AccountId);
+					PublicData->Members.Remove(LeavingMember->AccountId);
+					LobbyEvents.OnLobbyMemberLeft.Broadcast(FLobbyMemberLeft{PublicData, LeavingMember, ELobbyMemberLeaveReason::Left});
+				}
+
+				LobbyEvents.OnLobbyLeft.Broadcast(FLobbyLeft{PublicData});
+			}
+		}
+	}
+
+	return Result;
+}
+
+TSet<FLobbyAttributeId> FClientLobbyData::ApplyAttributeUpdateFromSnapshot(
+	TMap<FLobbyAttributeId, FLobbyVariant>&& AttributeSnapshot,
+	TMap<FLobbyAttributeId, FLobbyVariant>& ExistingAttributes)
+{
+	TSet<FLobbyAttributeId> Result;
+
+	// Clear the attributes.
+	{
+		// Determine which attributes were removed.
+		for (const TPair<FLobbyAttributeId, FLobbyVariant>& Attribute : ExistingAttributes)
+		{
+			if (!AttributeSnapshot.Contains(Attribute.Key))
+			{
+				Result.Add(Attribute.Key);
+			}
+		}
+
+		// Remove them.
+		for (const FLobbyAttributeId& AttributeId : Result)
+		{
+			ExistingAttributes.Remove(AttributeId);
+		}
+	}
+
+	// Apply attribute additions and mutations.
+	{
+		for (TPair<FLobbyAttributeId, FLobbyVariant>& Attribute : AttributeSnapshot)
+		{
+			if (FLobbyVariant* ExistingAttributeData = ExistingAttributes.Find(Attribute.Key))
+			{
+				if (*ExistingAttributeData != Attribute.Value)
+				{
+					Result.Add(MoveTemp(Attribute.Key));
+					*ExistingAttributeData = MoveTemp(Attribute.Value);
+				}
+			}
+			else
+			{
+				Result.Add(Attribute.Key);
+				ExistingAttributes.Add(MoveTemp(Attribute));
+			}
+		}
+	}
+
+	return Result;
+}
+
+TSet<FLobbyAttributeId> FClientLobbyData::ApplyAttributeUpdateFromChanges(
+	TMap<FLobbyAttributeId, FLobbyVariant>&& MutatedAttributes,
+	TSet<FLobbyAttributeId>&& ClearedAttributes,
+	TMap<FLobbyAttributeId, FLobbyVariant>& ExistingAttributes)
+{
+	TSet<FLobbyAttributeId> Result;
+
+	// Apply attribute additions and mutations.
+	{
+		for (TPair<FLobbyAttributeId, FLobbyVariant>& Attribute : MutatedAttributes)
+		{
+			if (FLobbyVariant* ExistingAttributeData = ExistingAttributes.Find(Attribute.Key))
+			{
+				if (*ExistingAttributeData != Attribute.Value)
+				{
+					Result.Add(MoveTemp(Attribute.Key));
+					*ExistingAttributeData = MoveTemp(Attribute.Value);
+				}
+			}
+			else
+			{
+				Result.Add(Attribute.Key);
+				ExistingAttributes.Add(MoveTemp(Attribute));
+			}
+		}
+	}
+
+	// Clear the attributes.
+	{
+		for (FLobbyAttributeId ClearedAttributeId : ClearedAttributes)
+		{
+			if (ExistingAttributes.Remove(ClearedAttributeId) > 0)
+			{
+				Result.Add(ClearedAttributeId);
+			}
+		}
+	}
+
+	return Result;
+}
+
+/* UE::Online */ }
diff --git a/Engine/Plugins/Online/Experimental/OnlineServices/Source/OnlineServicesCommon/Public/Online/LobbiesCommon.h b/Engine/Plugins/Online/Experimental/OnlineServices/Source/OnlineServicesCommon/Public/Online/LobbiesCommon.h
--- a/Engine/Plugins/Online/Experimental/OnlineServices/Source/OnlineServicesCommon/Public/Online/LobbiesCommon.h
+++ b/Engine/Plugins/Online/Experimental/OnlineServices/Source/OnlineServicesCommon/Public/Online/LobbiesCommon.h
@@ -4,188 +4,13 @@
 
 #include "Online/Lobbies.h"
 #include "Online/OnlineComponent.h"
-
-#define LOBBIES_FUNCTIONAL_TEST_ENABLED !UE_BUILD_SHIPPING
+#include "Online/LobbiesCommonTypes.h"
 
 namespace UE::Online {
 
 class FOnlineServicesCommon;
 class FAccountInfo;
 
-#if LOBBIES_FUNCTIONAL_TEST_ENABLED
-struct FFunctionalTestLobbies
-{
-	static constexpr TCHAR Name[] = TEXT("FunctionalTest");
-
-	struct Params
-	{
-	};
-
-	struct Result
-	{
-	};
-};
-#endif
-
-struct FLobbyConfig
-{
-	TArray<FLobbySchemaId> RegisteredSchemas;
-};
-
-struct FLobbySchemaAttributeConfig
-{
-	// The attribute name.
-	FLobbyAttributeId Name;
-	
-	// Override default service attribute field assignment.
-	// Specifying an override allows custom grouping of attributes within a
-	// platform attribute. This grouping is useful when a set of attributes are
-	// known to update at the same cadence.
-	TOptional<FLobbyAttributeId> ServiceAttributeFieldId;
-	
-	// The size in bytes of the attribute value.
-	// A schema will fail validation if all attributes will not fit within the
-	// platform's lobby fields.
-	uint32 MaxByteSize = 0;
-
-	// Control whether the attribute is visible to players who are not joined
-	// to the lobby.
-	ELobbyAttributeVisibility Visibility = ELobbyAttributeVisibility::Private;
-
-};
-
-struct FLobbySchemaConfig
-{
-	// The schema name.
-	FLobbySchemaId SchemaName;
-
-	// The optional base schema name.
-	FLobbySchemaId BaseSchemaName;
-	
-	// The definitions for attributes attached to a lobby.
-	TArray<FLobbySchemaAttributeConfig> LobbyAttributes;
-
-	// The definitions for attributes attached to a lobby member.
-	TArray<FLobbySchemaAttributeConfig> LobbyMemberAttributes;
-};
-
-class FLobbyServiceAttributeChanges final
-{
-public:
-	TMap<FLobbyAttributeId, FLobbyVariant> MutatedAttributes;
-	TSet<FLobbyAttributeId> ClearedAttributes;
-};
-
-class FLobbyClientAttributeChanges final
-{
-public:
-	// Apply attribute changes to attribute parameter.
-	// Applying the changes will clear MutatedAttributes and ClearedAttributes.
-	// Returns a list of the attributes which were changed.
-	TSet<FLobbyAttributeId> Apply(TMap<FLobbyAttributeId, FLobbyVariant>& InOutAttributes)
-	{
-		TSet<FLobbyAttributeId> Changes;
-		Changes.Reserve(MutatedAttributes.Num() + ClearedAttributes.Num());
-
-		for (TPair<FLobbyAttributeId, FLobbyVariant>& Attribute : MutatedAttributes)
-		{
-			Changes.Add(Attribute.Key);
-			InOutAttributes.Emplace(Attribute.Key, MoveTemp(Attribute.Value));
-		}
-
-		for (FLobbyAttributeId AttributeId : ClearedAttributes)
-		{
-			Changes.Add(AttributeId);
-			InOutAttributes.Remove(AttributeId);
-		}
-
-		MutatedAttributes.Empty();
-		ClearedAttributes.Empty();
-
-		return Changes;
-	}
-
-	TMap<FLobbyAttributeId, FLobbyVariant> MutatedAttributes;
-	TSet<FLobbyAttributeId> ClearedAttributes;
-};
-
-struct FLobbyEvents final
-{
-	TOnlineEventCallable<void(const FLobbyJoined&)> OnLobbyJoined;
-	TOnlineEventCallable<void(const FLobbyLeft&)> OnLobbyLeft;
-	TOnlineEventCallable<void(const FLobbyMemberJoined&)> OnLobbyMemberJoined;
-	TOnlineEventCallable<void(const FLobbyMemberLeft&)> OnLobbyMemberLeft;
-	TOnlineEventCallable<void(const FLobbyLeaderChanged&)> OnLobbyLeaderChanged;
-	TOnlineEventCallable<void(const FLobbySchemaChanged&)> OnLobbySchemaChanged;
-	TOnlineEventCallable<void(const FLobbyAttributesChanged&)> OnLobbyAttributesChanged;
-	TOnlineEventCallable<void(const FLobbyMemberAttributesChanged&)> OnLobbyMemberAttributesChanged;
-	TOnlineEventCallable<void(const FLobbyInvitationAdded&)> OnLobbyInvitationAdded;
-	TOnlineEventCallable<void(const FLobbyInvitationRemoved&)> OnLobbyInvitationRemoved;
-};
-
-class FLobbySchema final
-{
-public:
-	static TSharedPtr<FLobbySchema> Create(FLobbySchemaConfig LobbySchemaConfig);
-
-	TSharedRef<FLobbyServiceAttributeChanges> TranslateLobbyAttributes(const TMap<FLobbyAttributeId, FLobbyVariant>& LobbyAttributes, const FLobbyClientAttributeChanges& ClientAttributeChanges) const;
-	TSharedRef<FLobbyClientAttributeChanges> TranslateLobbyAttributes(const FLobbyServiceAttributeChanges& ServiceAttributeChanges) const;
-	TSharedRef<FLobbyServiceAttributeChanges> TranslateLobbyMemberAttributes(const TMap<FLobbyAttributeId, FLobbyVariant>& LobbyAttributes, const FLobbyClientAttributeChanges& ClientAttributeChanges) const;
-	TSharedRef<FLobbyClientAttributeChanges> TranslateLobbyMemberAttributes(const FLobbyServiceAttributeChanges& ServiceAttributeChanges) const;
-
-private:
-	FLobbySchemaId SchemaId;
-};
-
-class FLobbySchemaRegistry
-{
-public:
-	bool Initialize(TArray<FLobbySchemaConfig> LobbySchemaConfigs);
-	TSharedPtr<FLobbySchema> FindSchema(FLobbySchemaId SchemaId);
-
-private:
-	bool RegisterSchema(FLobbySchemaConfig LobbySchemaConfig);
-
-	TMap<FLobbySchemaId, TSharedRef<FLobbySchema>> RegisteredSchemas;
-};
-
-// Todo: put this somewhere else
-template <typename AwaitedType>
-TFuture<TArray<AwaitedType>> WhenAll(TArray<TFuture<AwaitedType>>&& Futures)
-{
-	struct FWhenAllState
-	{
-		TArray<TFuture<AwaitedType>> Futures;
-		TArray<AwaitedType> Results;
-		TPromise<TArray<AwaitedType>> FinalPromise;
-	};
-
-	if (Futures.IsEmpty())
-	{
-		return MakeFulfilledPromise<TArray<AwaitedType>>().GetFuture();
-	}
-	else
-	{
-		TSharedRef<FWhenAllState> WhenAllState = MakeShared<FWhenAllState>();
-		WhenAllState->Futures = MoveTemp(Futures);
-
-		for (TFuture<AwaitedType>& Future : WhenAllState->Futures)
-		{
-			Future.Then([WhenAllState](TFuture<AwaitedType>&& AwaitedResult)
-			{
-				WhenAllState->Results.Emplace(MoveTempIfPossible(AwaitedResult.Get()));
-
-				if (WhenAllState->Futures.Num() == WhenAllState->Results.Num())
-				{
-					WhenAllState->FinalPromise.EmplaceValue(MoveTemp(WhenAllState->Results));
-				}
-			});
-		}
-
-		return WhenAllState->FinalPromise.GetFuture();
-	}
-}
-
 class ONLINESERVICESCOMMON_API FLobbiesCommon : public TOnlineComponent<ILobbies>
 {
 public:
@@ -226,333 +51,13 @@ public:
 	virtual TOnlineEvent<void(const FLobbyInvitationRemoved&)> OnLobbyInvitationRemoved() override;
 
 protected:
-	TFuture<TDefaultErrorResultInternal<FOnlineLobbyIdHandle>> AwaitInvitation(
-		FOnlineAccountIdHandle TargetAccountId,
-		FOnlineLobbyIdHandle LobbyId,
-		float TimeoutSeconds);
-
 #if LOBBIES_FUNCTIONAL_TEST_ENABLED
 	TOnlineAsyncOpHandle<FFunctionalTestLobbies> FunctionalTest(FFunctionalTestLobbies::Params&& Params);
-
-	struct FFunctionalTestLoginUser
-	{
-		static constexpr TCHAR Name[] = TEXT("FunctionalTestLoginUser");
-
-		struct Params
-		{
-			FPlatformUserId PlatformUserId;
-			FString Type;
-			FString Id;
-			FString Token;
-		};
-
-		struct Result
-		{
-			TSharedPtr<FAccountInfo> AccountInfo;
-		};
-	};
-	TFuture<TOnlineResult<FFunctionalTestLoginUser>> FunctionalTestLoginUser(FFunctionalTestLoginUser::Params&& Params);
-
-	struct FFunctionalTestLogoutUser
-	{
-		static constexpr TCHAR Name[] = TEXT("FunctionalTestLogoutUser");
-
-		struct Params
-		{
-			FPlatformUserId PlatformUserId;
-		};
-
-		struct Result
-		{
-		};
-	};
-	TFuture<TOnlineResult<FFunctionalTestLogoutUser>> FunctionalTestLogoutUser(FFunctionalTestLogoutUser::Params&& Params);
-
-	struct FFunctionalTestLogoutAllUsers
-	{
-		static constexpr TCHAR Name[] = TEXT("FunctionalTestLogoutAllUsers");
-
-		struct Params
-		{
-		};
-
-		struct Result
-		{
-		};
-	};
-	TFuture<TOnlineResult<FFunctionalTestLogoutAllUsers>> FunctionalTestLogoutAllUsers(FFunctionalTestLogoutAllUsers::Params&& Params);
-#endif
-
-	template <typename SecondaryOpType, typename OpType, typename Function>
-	TFunction<TFuture<void>(TOnlineAsyncOp<OpType>&, typename SecondaryOpType::Params&&)> ConsumeStepResult(TOnlineAsyncOp<OpType>& InAsyncOp, Function Func)
-	{
-		return [this, Func](TOnlineAsyncOp<OpType>& InAsyncOp, typename SecondaryOpType::Params&& InParams) mutable
-		{
-			TPromise<void> Promise;
-			auto Future = Promise.GetFuture();
-
-			UE_LOG(LogTemp, Log, TEXT("[FLobbiesCommon::ConsumeStepResult] Starting secondary operation %s."), SecondaryOpType::Name);
-
-			(this->*Func)(MoveTempIfPossible(InParams))
-			.Then([Promise = MoveTemp(Promise), Op = InAsyncOp.AsShared()](TFuture<TOnlineResult<SecondaryOpType>>&& Future) mutable
-			{
-				if (Future.Get().IsError())
-				{
-					UE_LOG(LogTemp, Log, TEXT("[FLobbiesCommon::ConsumeStepResult] Failed secondary operation %s. Error: %s"), SecondaryOpType::Name, *Future.Get().GetErrorValue().GetLogString());
-					Op->SetError(Errors::RequestFailure(MoveTempIfPossible(Future.Get().GetErrorValue())));
-				}
-				else
-				{
-					UE_LOG(LogTemp, Log, TEXT("[FLobbiesCommon::ConsumeStepResult] Completed secondary operation %s."), SecondaryOpType::Name);
-				}
-				Promise.EmplaceValue();
-			});
-
-			return Future;
-		};
-	}
-
-	template <typename SecondaryOpType, typename OpType, typename Function>
-	TFunction<TFuture<void>(TOnlineAsyncOp<OpType>&)> ConsumeStepResult(TOnlineAsyncOp<OpType>& InAsyncOp, Function Func, typename SecondaryOpType::Params&& InParams)
-	{
-		return [this, Func, InParams = MoveTempIfPossible(InParams)](TOnlineAsyncOp<OpType>& InAsyncOp) mutable
-		{
-			TPromise<void> Promise;
-			auto Future = Promise.GetFuture();
-
-			UE_LOG(LogTemp, Log, TEXT("[FLobbiesCommon::ConsumeStepResult] Starting secondary operation %s."), SecondaryOpType::Name);
-
-			(this->*Func)(MoveTempIfPossible(InParams))
-			.Then([Promise = MoveTemp(Promise), Op = InAsyncOp.AsShared()](TFuture<TOnlineResult<SecondaryOpType>>&& Future) mutable
-			{
-				if (Future.Get().IsError())
-				{
-					UE_LOG(LogTemp, Log, TEXT("[FLobbiesCommon::ConsumeStepResult] Failed secondary operation %s. Error: %s"), SecondaryOpType::Name, *Future.Get().GetErrorValue().GetLogString());
-					Op->SetError(Errors::RequestFailure(MoveTempIfPossible(Future.Get().GetErrorValue())));
-				}
-				else
-				{
-					UE_LOG(LogTemp, Log, TEXT("[FLobbiesCommon::ConsumeStepResult] Completed secondary operation %s."), SecondaryOpType::Name);
-				}
-				Promise.EmplaceValue();
-			});
-
-			return Future;
-		};
-	}
-
-	template <typename SecondaryOpType, typename OpType, typename Function>
-	TFunction<TFuture<void>(TOnlineAsyncOp<OpType>&, typename SecondaryOpType::Params&&)> CaptureStepResult(TOnlineAsyncOp<OpType>& InAsyncOp, const FString& ResultKey, Function Func)
-	{
-		return [this, ResultKey, Func](TOnlineAsyncOp<OpType>& InAsyncOp, typename SecondaryOpType::Params&& InParams) mutable
-		{
-			TPromise<void> Promise;
-			auto Future = Promise.GetFuture();
-
-			UE_LOG(LogTemp, Log, TEXT("[FLobbiesCommon::CaptureStepResult] Starting secondary operation %s."), SecondaryOpType::Name);
-
-			(this->*Func)(MoveTempIfPossible(InParams))
-			.Then([Promise = MoveTemp(Promise), Op = InAsyncOp.AsShared(), ResultKey](TFuture<TOnlineResult<SecondaryOpType>>&& Future) mutable
-			{
-				if (Future.Get().IsError())
-				{
-					UE_LOG(LogTemp, Log, TEXT("[FLobbiesCommon::CaptureStepResult] Failed secondary operation %s. Error: %s"), SecondaryOpType::Name, *Future.Get().GetErrorValue().GetLogString());
-					Op->SetError(Errors::RequestFailure(MoveTempIfPossible(Future.Get().GetErrorValue())));
-				}
-				else
-				{
-					UE_LOG(LogTemp, Log, TEXT("[FLobbiesCommon::CaptureStepResult] Captured secondary operation %s as %s"), SecondaryOpType::Name, *ResultKey);
-					Op->Data.Set(ResultKey, MoveTempIfPossible(Future.Get().GetOkValue()));
-				}
-				Promise.EmplaceValue();
-			});
-
-			return Future;
-		};
-	}
-
-	template <typename SecondaryOpType, typename OpType, typename Function>
-	TFunction<TFuture<void>(TOnlineAsyncOp<OpType>&)> CaptureStepResult(TOnlineAsyncOp<OpType>& InAsyncOp, const FString& ResultKey, Function Func, typename SecondaryOpType::Params&& InParams)
-	{
-		return [this, ResultKey, Func, InParams = MoveTempIfPossible(InParams)](TOnlineAsyncOp<OpType>& InAsyncOp) mutable
-		{
-			TPromise<void> Promise;
-			auto Future = Promise.GetFuture();
-
-			UE_LOG(LogTemp, Log, TEXT("[FLobbiesCommon::CaptureStepResult] Starting secondary operation %s."), SecondaryOpType::Name);
-
-			(this->*Func)(MoveTempIfPossible(InParams))
-			.Then([Promise = MoveTemp(Promise), Op = InAsyncOp.AsShared(), ResultKey](TFuture<TOnlineResult<SecondaryOpType>>&& Future) mutable
-			{
-				if (Future.Get().IsError())
-				{
-					UE_LOG(LogTemp, Log, TEXT("[FLobbiesCommon::CaptureStepResult] Failed secondary operation %s. Error: %s"), SecondaryOpType::Name, *Future.Get().GetErrorValue().GetLogString());
-					Op->SetError(Errors::RequestFailure(MoveTempIfPossible(Future.Get().GetErrorValue())));
-				}
-				else
-				{
-					UE_LOG(LogTemp, Log, TEXT("[FLobbiesCommon::CaptureStepResult] Captured secondary operation %s as %s"), SecondaryOpType::Name, *ResultKey);
-					Op->Data.Set(ResultKey, MoveTempIfPossible(Future.Get().GetOkValue()));
-				}
-				Promise.EmplaceValue();
-			});
-
-			return Future;
-		};
-	}
-
-	template <typename SecondaryOpType, typename OpType, typename Function>
-	TFunction<TFuture<void>(TOnlineAsyncOp<OpType>&, typename SecondaryOpType::Params&&)> ConsumeOperationStepResult(TOnlineAsyncOp<OpType>& InAsyncOp, Function Func)
-	{
-		return [this, Func](TOnlineAsyncOp<OpType>& InAsyncOp, typename SecondaryOpType::Params&& InParams) mutable -> TFuture<void>
-		{
-			auto Promise = MakeShared<TPromise<void>>();
-			auto Future = Promise->GetFuture();
-
-			UE_LOG(LogTemp, Log, TEXT("[FLobbiesCommon::ConsumeOperationStepResult] Starting secondary operation %s."), SecondaryOpType::Name);
-
-			(this->*Func)(MoveTempIfPossible(InParams))
-			.OnComplete([Promise, Op = InAsyncOp.AsShared()](const TOnlineResult<SecondaryOpType>& Result) mutable -> void
-			{
-				if (Result.IsError())
-				{
-					UE_LOG(LogTemp, Log, TEXT("[FLobbiesCommon::ConsumeOperationStepResult] Failed secondary operation %s. Error: %s"), SecondaryOpType::Name, *Result.GetErrorValue().GetLogString());
-					Op->SetError(Errors::RequestFailure(MoveTempIfPossible(Result.GetErrorValue())));
-				}
-				else
-				{
-					UE_LOG(LogTemp, Log, TEXT("[FLobbiesCommon::ConsumeOperationStepResult] Completed secondary operation %s."), SecondaryOpType::Name);
-				}
-				Promise->EmplaceValue();
-			});
-
-			return Future;
-		};
-	}
-
-	template <typename SecondaryOpType, typename OpType, typename Function>
-	TFunction<TFuture<void>(TOnlineAsyncOp<OpType>&)> ConsumeOperationStepResult(TOnlineAsyncOp<OpType>& InAsyncOp, Function Func, typename SecondaryOpType::Params&& InParams)
-	{
-		return [this, Func, InParams = MoveTempIfPossible(InParams)](TOnlineAsyncOp<OpType>& InAsyncOp) mutable
-		{
-			auto Promise = MakeShared<TPromise<void>>();
-			auto Future = Promise->GetFuture();
-
-			UE_LOG(LogTemp, Log, TEXT("[FLobbiesCommon::ConsumeOperationStepResult] Starting secondary operation %s."), SecondaryOpType::Name);
-
-			(this->*Func)(MoveTempIfPossible(InParams))
-			.OnComplete([Promise, Op = InAsyncOp.AsShared()](const TOnlineResult<SecondaryOpType>& Result) mutable
-			{
-				if (Result.IsError())
-				{
-					UE_LOG(LogTemp, Log, TEXT("[FLobbiesCommon::ConsumeOperationStepResult] Failed secondary operation %s. Error: %s"), SecondaryOpType::Name, *Result.GetErrorValue().GetLogString());
-					Op->SetError(Errors::RequestFailure(MoveTempIfPossible(Result.GetErrorValue())));
-				}
-				else
-				{
-					UE_LOG(LogTemp, Log, TEXT("[FLobbiesCommon::ConsumeOperationStepResult] Completed secondary operation %s."), SecondaryOpType::Name);
-				}
-				Promise->EmplaceValue();
-			});
-
-			return Future;
-		};
-	}
-
-	template <typename SecondaryOpType, typename OpType, typename Function>
-	TFunction<TFuture<void>(TOnlineAsyncOp<OpType>&, typename SecondaryOpType::Params&&)> CaptureOperationStepResult(TOnlineAsyncOp<OpType>& InAsyncOp, const FString& ResultKey, Function Func)
-	{
-		return [this, ResultKey, Func](TOnlineAsyncOp<OpType>& InAsyncOp, typename SecondaryOpType::Params&& InParams) mutable -> TFuture<void>
-		{
-			auto Promise = MakeShared<TPromise<void>>();
-			auto Future = Promise->GetFuture();
-
-			UE_LOG(LogTemp, Log, TEXT("[FLobbiesCommon::CaptureOperationStepResult] Starting secondary operation %s."), SecondaryOpType::Name);
-
-			(this->*Func)(MoveTempIfPossible(InParams))
-			.OnComplete([Promise, Op = InAsyncOp.AsShared(), ResultKey](const TOnlineResult<SecondaryOpType>& Result) mutable -> void
-			{
-				if (Result.IsError())
-				{
-					UE_LOG(LogTemp, Log, TEXT("[FLobbiesCommon::CaptureOperationStepResult] Failed secondary operation %s. Error: %s"), SecondaryOpType::Name, *Result.GetErrorValue().GetLogString());
-					Op->SetError(Errors::RequestFailure(MoveTempIfPossible(Result.GetErrorValue())));
-				}
-				else
-				{
-					UE_LOG(LogTemp, Log, TEXT("[FLobbiesCommon::CaptureOperationStepResult] Captured secondary operation %s as %s"), SecondaryOpType::Name, *ResultKey);
-					Op->Data.Set(ResultKey, Result.GetOkValue());
-				}
-				Promise->EmplaceValue();
-			});
-
-			return Future;
-		};
-	}
-
-	template <typename SecondaryOpType, typename OpType, typename Function>
-	TFunction<TFuture<void>(TOnlineAsyncOp<OpType>&)> CaptureOperationStepResult(TOnlineAsyncOp<OpType>& InAsyncOp, const FString& ResultKey, Function Func, typename SecondaryOpType::Params&& InParams)
-	{
-		return [this, ResultKey, Func, InParams = MoveTempIfPossible(InParams)](TOnlineAsyncOp<OpType>& InAsyncOp) mutable
-		{
-			auto Promise = MakeShared<TPromise<void>>();
-			auto Future = Promise->GetFuture();
-
-			UE_LOG(LogTemp, Log, TEXT("[FLobbiesCommon::CaptureOperationStepResult] Starting secondary operation %s."), SecondaryOpType::Name);
-
-			(this->*Func)(MoveTempIfPossible(InParams))
-			.OnComplete([Promise, Op = InAsyncOp.AsShared(), ResultKey](const TOnlineResult<SecondaryOpType>& Result) mutable
-			{
-				if (Result.IsError())
-				{
-					UE_LOG(LogTemp, Log, TEXT("[FLobbiesCommon::CaptureOperationStepResult] Failed secondary operation %s. Error: %s"), SecondaryOpType::Name, *Result.GetErrorValue().GetLogString());
-					Op->SetError(Errors::RequestFailure(MoveTempIfPossible(Result.GetErrorValue())));
-				}
-				else
-				{
-					UE_LOG(LogTemp, Log, TEXT("[FLobbiesCommon::CaptureOperationStepResult] Captured secondary operation %s as %s"), SecondaryOpType::Name, *ResultKey);
-					Op->Data.Set(ResultKey, Result.GetOkValue());
-				}
-				Promise->EmplaceValue();
-			});
-
-			return Future;
-		};
-	}
+#endif // LOBBIES_FUNCTIONAL_TEST_ENABLED
 
 	FLobbyEvents LobbyEvents;
-
 	TSharedPtr<FLobbySchemaRegistry> LobbySchemaRegistry;
 	TSharedPtr<FLobbySchema> ServiceSchema;
 };
 
-namespace Meta {
-
-BEGIN_ONLINE_STRUCT_META(FLobbyConfig)
-	ONLINE_STRUCT_FIELD(FLobbyConfig, RegisteredSchemas)
-END_ONLINE_STRUCT_META()
-
-
-BEGIN_ONLINE_STRUCT_META(FLobbySchemaAttributeConfig)
-	ONLINE_STRUCT_FIELD(FLobbySchemaAttributeConfig, Name),
-	ONLINE_STRUCT_FIELD(FLobbySchemaAttributeConfig, ServiceAttributeFieldId),
-	ONLINE_STRUCT_FIELD(FLobbySchemaAttributeConfig, MaxByteSize),
-	ONLINE_STRUCT_FIELD(FLobbySchemaAttributeConfig, Visibility)
-END_ONLINE_STRUCT_META()
-
-BEGIN_ONLINE_STRUCT_META(FLobbySchemaConfig)
-	ONLINE_STRUCT_FIELD(FLobbySchemaConfig, BaseSchemaName),
-	ONLINE_STRUCT_FIELD(FLobbySchemaConfig, LobbyAttributes),
-	ONLINE_STRUCT_FIELD(FLobbySchemaConfig, LobbyMemberAttributes)
-END_ONLINE_STRUCT_META()
-
-#if LOBBIES_FUNCTIONAL_TEST_ENABLED
-	BEGIN_ONLINE_STRUCT_META(FFunctionalTestLobbies::Params)
-	END_ONLINE_STRUCT_META()
-
-	BEGIN_ONLINE_STRUCT_META(FFunctionalTestLobbies::Result)
-	END_ONLINE_STRUCT_META()
-#endif
-
-/* Meta */ }
-
 /* UE::Online */ }
diff --git a/Engine/Plugins/Online/Experimental/OnlineServices/Source/OnlineServicesCommon/Public/Online/LobbiesCommonTypes.h b/Engine/Plugins/Online/Experimental/OnlineServices/Source/OnlineServicesCommon/Public/Online/LobbiesCommonTypes.h
new file mode 100644
--- /dev/null
+++ b/Engine/Plugins/Online/Experimental/OnlineServices/Source/OnlineServicesCommon/Public/Online/LobbiesCommonTypes.h
@@ -0,0 +1,343 @@
+// Copyright Epic Games, Inc. All Rights Reserved.
+
+#pragma once
+
+#include "Async/Future.h"
+#include "Online/Lobbies.h"
+
+namespace UE::Online {
+
+#define LOBBIES_FUNCTIONAL_TEST_ENABLED !UE_BUILD_SHIPPING
+
+class FOnlineServicesCommon;
+class FAccountInfo;
+
+struct FLobbyConfig
+{
+	TArray<FLobbySchemaId> RegisteredSchemas;
+};
+
+struct FLobbySchemaAttributeConfig
+{
+	// The attribute name.
+	FLobbyAttributeId Name;
+	
+	// Override default service attribute field assignment.
+	// Specifying an override allows custom grouping of attributes within a
+	// platform attribute. This grouping is useful when a set of attributes are
+	// known to update at the same cadence.
+	TOptional<FLobbyAttributeId> ServiceAttributeFieldId;
+	
+	// The size in bytes of the attribute value.
+	// A schema will fail validation if all attributes will not fit within the
+	// platform's lobby fields.
+	uint32 MaxByteSize = 0;
+
+	// Control whether the attribute is visible to players who are not joined
+	// to the lobby.
+	ELobbyAttributeVisibility Visibility = ELobbyAttributeVisibility::Private;
+
+};
+
+struct FLobbySchemaConfig
+{
+	// The schema name.
+	FLobbySchemaId SchemaName;
+
+	// The optional base schema name.
+	FLobbySchemaId BaseSchemaName;
+	
+	// The definitions for attributes attached to a lobby.
+	TArray<FLobbySchemaAttributeConfig> LobbyAttributes;
+
+	// The definitions for attributes attached to a lobby member.
+	TArray<FLobbySchemaAttributeConfig> LobbyMemberAttributes;
+};
+
+class FLobbyServiceAttributeChanges final
+{
+public:
+	TMap<FLobbyAttributeId, FLobbyVariant> MutatedAttributes;
+	TSet<FLobbyAttributeId> ClearedAttributes;
+};
+
+class FLobbyClientAttributeChanges final
+{
+public:
+	// Apply attribute changes to attribute parameter.
+	// Applying the changes will clear MutatedAttributes and ClearedAttributes.
+	// Returns a list of the attributes which were changed.
+	TSet<FLobbyAttributeId> Apply(TMap<FLobbyAttributeId, FLobbyVariant>& InOutAttributes)
+	{
+		TSet<FLobbyAttributeId> Changes;
+		Changes.Reserve(MutatedAttributes.Num() + ClearedAttributes.Num());
+
+		for (TPair<FLobbyAttributeId, FLobbyVariant>& Attribute : MutatedAttributes)
+		{
+			Changes.Add(Attribute.Key);
+			InOutAttributes.Emplace(Attribute.Key, MoveTemp(Attribute.Value));
+		}
+
+		for (FLobbyAttributeId AttributeId : ClearedAttributes)
+		{
+			Changes.Add(AttributeId);
+			InOutAttributes.Remove(AttributeId);
+		}
+
+		MutatedAttributes.Empty();
+		ClearedAttributes.Empty();
+
+		return Changes;
+	}
+
+	TMap<FLobbyAttributeId, FLobbyVariant> MutatedAttributes;
+	TSet<FLobbyAttributeId> ClearedAttributes;
+};
+
+struct FLobbyEvents final
+{
+	TOnlineEventCallable<void(const FLobbyJoined&)> OnLobbyJoined;
+	TOnlineEventCallable<void(const FLobbyLeft&)> OnLobbyLeft;
+	TOnlineEventCallable<void(const FLobbyMemberJoined&)> OnLobbyMemberJoined;
+	TOnlineEventCallable<void(const FLobbyMemberLeft&)> OnLobbyMemberLeft;
+	TOnlineEventCallable<void(const FLobbyLeaderChanged&)> OnLobbyLeaderChanged;
+	TOnlineEventCallable<void(const FLobbySchemaChanged&)> OnLobbySchemaChanged;
+	TOnlineEventCallable<void(const FLobbyAttributesChanged&)> OnLobbyAttributesChanged;
+	TOnlineEventCallable<void(const FLobbyMemberAttributesChanged&)> OnLobbyMemberAttributesChanged;
+	TOnlineEventCallable<void(const FLobbyInvitationAdded&)> OnLobbyInvitationAdded;
+	TOnlineEventCallable<void(const FLobbyInvitationRemoved&)> OnLobbyInvitationRemoved;
+};
+
+class ONLINESERVICESCOMMON_API FLobbySchema final
+{
+public:
+	static TSharedPtr<FLobbySchema> Create(FLobbySchemaConfig LobbySchemaConfig);
+
+	TSharedRef<FLobbyServiceAttributeChanges> TranslateLobbyAttributes(const TMap<FLobbyAttributeId, FLobbyVariant>& LobbyAttributes, const FLobbyClientAttributeChanges& ClientAttributeChanges) const;
+	TSharedRef<FLobbyClientAttributeChanges> TranslateLobbyAttributes(const FLobbyServiceAttributeChanges& ServiceAttributeChanges) const;
+	TSharedRef<FLobbyServiceAttributeChanges> TranslateLobbyMemberAttributes(const TMap<FLobbyAttributeId, FLobbyVariant>& LobbyAttributes, const FLobbyClientAttributeChanges& ClientAttributeChanges) const;
+	TSharedRef<FLobbyClientAttributeChanges> TranslateLobbyMemberAttributes(const FLobbyServiceAttributeChanges& ServiceAttributeChanges) const;
+
+private:
+	FLobbySchemaId SchemaId;
+};
+
+class FLobbySchemaRegistry
+{
+public:
+	bool Initialize(TArray<FLobbySchemaConfig> LobbySchemaConfigs);
+	TSharedPtr<FLobbySchema> FindSchema(FLobbySchemaId SchemaId);
+
+private:
+	bool RegisterSchema(FLobbySchemaConfig LobbySchemaConfig);
+
+	TMap<FLobbySchemaId, TSharedRef<FLobbySchema>> RegisteredSchemas;
+};
+
+// Todo: put this somewhere else
+template <typename AwaitedType>
+TFuture<TArray<AwaitedType>> WhenAll(TArray<TFuture<AwaitedType>>&& Futures)
+{
+	struct FWhenAllState
+	{
+		TArray<TFuture<AwaitedType>> Futures;
+		TArray<AwaitedType> Results;
+		TPromise<TArray<AwaitedType>> FinalPromise;
+	};
+
+	if (Futures.IsEmpty())
+	{
+		return MakeFulfilledPromise<TArray<AwaitedType>>().GetFuture();
+	}
+	else
+	{
+		TSharedRef<FWhenAllState> WhenAllState = MakeShared<FWhenAllState>();
+		WhenAllState->Futures = MoveTemp(Futures);
+
+		for (TFuture<AwaitedType>& Future : WhenAllState->Futures)
+		{
+			Future.Then([WhenAllState](TFuture<AwaitedType>&& AwaitedResult)
+			{
+				WhenAllState->Results.Emplace(MoveTempIfPossible(AwaitedResult.Get()));
+
+				if (WhenAllState->Futures.Num() == WhenAllState->Results.Num())
+				{
+					WhenAllState->FinalPromise.EmplaceValue(MoveTemp(WhenAllState->Results));
+				}
+			});
+		}
+
+		return WhenAllState->FinalPromise.GetFuture();
+	}
+}
+
+/**
+ * Local changes to a member to be applied to an existing lobby member snapshot.
+ */
+struct FClientLobbyMemberDataChanges
+{
+	/** New or changed attributes. */
+	TMap<FLobbyAttributeId, FLobbyVariant> MutatedAttributes;
+
+	/** Attributes to be cleared. */
+	TSet<FLobbyAttributeId> ClearedAttributes;
+};
+
+/**
+ * Local changes to a lobby to be applied to an existing lobby snapshot.
+ */
+struct FClientLobbyDataChanges
+{
+	/** Local name for the lobby. */
+	TOptional<FName> LocalName;
+
+	/** Setting for new join policy. */
+	TOptional<ELobbyJoinPolicy> JoinPolicy;
+
+	/** Setting for lobby ownership change. */
+	TOptional<FOnlineAccountIdHandle> OwnerAccountId;
+
+	/** Setting for lobby schema change. */
+	TOptional<FLobbySchemaId> LobbySchema;
+
+	/** New or changed attributes. */
+	TMap<FLobbyAttributeId, FLobbyVariant> MutatedAttributes;
+
+	/** Attributes to be cleared. */
+	TSet<FLobbyAttributeId> ClearedAttributes;
+
+	/** Members to be added or changed. */
+	TMap<FOnlineAccountIdHandle, TSharedRef<FClientLobbyMemberDataChanges>> MutatedMembers;
+
+	/** Members to be removed. */
+	TMap<FOnlineAccountIdHandle, ELobbyMemberLeaveReason> LeavingMembers;
+};
+
+/**
+* Lobby snapshot data.
+* All contained data has been translated from EOS types.
+* Attributes have had their schema transformations applied.
+*/
+struct FClientLobbySnapshot
+{
+	FOnlineAccountIdHandle OwnerAccountId;
+	FName SchemaName;
+	int32 MaxMembers;
+	ELobbyJoinPolicy JoinPolicy;
+	TMap<FLobbyAttributeId, FLobbyVariant> Attributes;
+	TSet<FOnlineAccountIdHandle> Members;
+};
+
+struct FClientLobbyMemberSnapshot : public FLobbyMember
+{
+	bool bIsLocalMember = false;
+};
+
+struct FApplyLobbyUpdateResult
+{
+	TArray<FOnlineAccountIdHandle> LeavingLocalMembers;
+};
+
+/** Lobby data as seen by the client. */
+struct ONLINESERVICESCOMMON_API FClientLobbyData final
+{
+public:
+	FClientLobbyData(FOnlineLobbyIdHandle LobbyId);
+
+	TSharedRef<const FLobby> GetPublicDataPtr() const { return PublicData; }
+	const FLobby& GetPublicData() const { return *PublicData; }
+
+	TSharedPtr<const FClientLobbyMemberSnapshot> GetMemberData(FOnlineAccountIdHandle MemberAccountId) const;
+
+	/**
+	 * Apply updated lobby data and generate changes.
+	 * The LeaveReason provides context for members who have left since the most recent snapshot.
+	 * Changing a lobby generates events for the local client.
+	 */
+	FApplyLobbyUpdateResult ApplyLobbyUpdateFromServiceSnapshot(
+		FClientLobbySnapshot&& LobbySnapshot,
+		TMap<FOnlineAccountIdHandle, TSharedRef<FClientLobbyMemberSnapshot>>&& LobbyMemberSnapshots,
+		TMap<FOnlineAccountIdHandle, ELobbyMemberLeaveReason>&& LeaveReasons = TMap<FOnlineAccountIdHandle, ELobbyMemberLeaveReason>(),
+		FLobbyEvents* LobbyEvents = nullptr);
+
+	/**
+	 * Apply updated lobby data and generate changes.
+	 * Changing a lobby generates events for the local client.
+	 */
+	FApplyLobbyUpdateResult ApplyLobbyUpdateFromLocalChanges(FClientLobbyDataChanges&& Changes, FLobbyEvents& LobbyEvents);
+
+private:
+
+	/** Apply changes to a set of attributes. Returns the set of attribute IDs which changed. */
+	TSet<FLobbyAttributeId> ApplyAttributeUpdateFromSnapshot(
+		TMap<FLobbyAttributeId, FLobbyVariant>&& AttributeSnapshot,
+		TMap<FLobbyAttributeId, FLobbyVariant>& ExistingAttributes);
+
+	/** Apply changes to a set of attributes. Returns the set of attribute IDs which changed. */
+	TSet<FLobbyAttributeId> ApplyAttributeUpdateFromChanges(
+		TMap<FLobbyAttributeId, FLobbyVariant>&& MutatedAttributes,
+		TSet<FLobbyAttributeId>&& ClearedAttributes,
+		TMap<FLobbyAttributeId, FLobbyVariant>& ExistingAttributes);
+
+	/**
+	 * The shared pointer given back to user code with lobby operation results and notifications.
+	 * Any changes to this data are immediately available to users.
+	 */
+	TSharedRef<FLobby> PublicData;
+
+	/** Mutable lobby member data storage. */
+	TMap<FOnlineAccountIdHandle, TSharedRef<FClientLobbyMemberSnapshot>> MemberDataStorage;
+
+	/**
+	 * Keep track of which members are local to the client.
+	 * When all local members have been removed all members will be removed.
+	 */
+	TSet<FOnlineAccountIdHandle> LocalMembers;
+};
+
+#if LOBBIES_FUNCTIONAL_TEST_ENABLED
+struct FFunctionalTestLobbies
+{
+	static constexpr TCHAR Name[] = TEXT("FunctionalTest");
+
+	struct Params
+	{
+	};
+
+	struct Result
+	{
+	};
+};
+#endif // LOBBIES_FUNCTIONAL_TEST_ENABLED
+
+namespace Meta {
+
+BEGIN_ONLINE_STRUCT_META(FLobbyConfig)
+	ONLINE_STRUCT_FIELD(FLobbyConfig, RegisteredSchemas)
+END_ONLINE_STRUCT_META()
+
+
+BEGIN_ONLINE_STRUCT_META(FLobbySchemaAttributeConfig)
+	ONLINE_STRUCT_FIELD(FLobbySchemaAttributeConfig, Name),
+	ONLINE_STRUCT_FIELD(FLobbySchemaAttributeConfig, ServiceAttributeFieldId),
+	ONLINE_STRUCT_FIELD(FLobbySchemaAttributeConfig, MaxByteSize),
+	ONLINE_STRUCT_FIELD(FLobbySchemaAttributeConfig, Visibility)
+END_ONLINE_STRUCT_META()
+
+BEGIN_ONLINE_STRUCT_META(FLobbySchemaConfig)
+	ONLINE_STRUCT_FIELD(FLobbySchemaConfig, BaseSchemaName),
+	ONLINE_STRUCT_FIELD(FLobbySchemaConfig, LobbyAttributes),
+	ONLINE_STRUCT_FIELD(FLobbySchemaConfig, LobbyMemberAttributes)
+END_ONLINE_STRUCT_META()
+
+#if LOBBIES_FUNCTIONAL_TEST_ENABLED
+BEGIN_ONLINE_STRUCT_META(FFunctionalTestLobbies::Params)
+END_ONLINE_STRUCT_META()
+
+BEGIN_ONLINE_STRUCT_META(FFunctionalTestLobbies::Result)
+END_ONLINE_STRUCT_META()
+#endif // LOBBIES_FUNCTIONAL_TEST_ENABLED
+
+/* Meta */ }
+
+/* UE::Online */ }
diff --git a/Engine/Plugins/Online/Experimental/OnlineServicesEOS/Source/Private/LobbiesEOSTypes.cpp b/Engine/Plugins/Online/Experimental/OnlineServicesEOS/Source/Private/LobbiesEOSTypes.cpp
--- a/Engine/Plugins/Online/Experimental/OnlineServicesEOS/Source/Private/LobbiesEOSTypes.cpp
+++ b/Engine/Plugins/Online/Experimental/OnlineServicesEOS/Source/Private/LobbiesEOSTypes.cpp
@@ -153,420 +153,6 @@ FLobbyAttributeTranslator<EAttributeTranslationType::FromService>::FLobbyAttribu
 	AttributeData = TPair<FLobbyAttributeId, FLobbyVariant>(MoveTemp(AttributeId), MoveTemp(VariantData));
 }
 
-FClientLobbyData::FClientLobbyData(FOnlineLobbyIdHandle LobbyId)
-	: PublicData(MakeShared<FLobby>())
-{
-	PublicData->LobbyId = LobbyId;
-}
-
-TSharedPtr<const FClientLobbyMemberSnapshot> FClientLobbyData::GetMemberData(FOnlineAccountIdHandle MemberAccountId) const
-{
-	const TSharedRef<FClientLobbyMemberSnapshot>* FoundMemberData = MemberDataStorage.Find(MemberAccountId);
-	return FoundMemberData ? *FoundMemberData : TSharedPtr<const FClientLobbyMemberSnapshot>();
-}
-
-FApplyLobbyUpdateResult FClientLobbyData::ApplyLobbyUpdateFromServiceSnapshot(
-	FClientLobbySnapshot&& LobbySnapshot,
-	TMap<FOnlineAccountIdHandle, TSharedRef<FClientLobbyMemberSnapshot>>&& LobbyMemberSnapshots,
-	TMap<FOnlineAccountIdHandle, ELobbyMemberLeaveReason>&& LeaveReasons,
-	FLobbyEvents* LobbyEvents)
-{
-	FApplyLobbyUpdateResult Result;
-
-	// Notify if schema changed.
-	// Schema change notification must be processed first as the schema affects how the client processes attributes.
-	if (LobbySnapshot.SchemaName != PublicData->SchemaName)
-	{
-		PublicData->SchemaName = LobbySnapshot.SchemaName;
-		if (LobbyEvents)
-		{
-			LobbyEvents->OnLobbySchemaChanged.Broadcast(FLobbySchemaChanged{PublicData});
-		}
-	}
-
-	// Handle lobby attribute changes.
-	TSet<FLobbyAttributeId> LobbyAttributeChanges = ApplyAttributeUpdateFromSnapshot(MoveTemp(LobbySnapshot.Attributes), PublicData->Attributes);
-	if (LobbyEvents && !LobbyAttributeChanges.IsEmpty())
-	{
-		LobbyEvents->OnLobbyAttributesChanged.Broadcast(FLobbyAttributesChanged{PublicData, MoveTemp(LobbyAttributeChanges)});
-	}
-
-	// Process members who left.
-	{
-		// Build list of leaving members.
-		TArray<TPair<TSharedRef<FClientLobbyMemberSnapshot>, ELobbyMemberLeaveReason>> LeavingMembers;
-		for (TPair<FOnlineAccountIdHandle, TSharedRef<FClientLobbyMemberSnapshot>>& MemberData : MemberDataStorage)
-		{
-			// Check if member exists in new snapshot.
-			if (LobbySnapshot.Members.Find(MemberData.Key) == nullptr)
-			{
-				// Member has left - check if they gave an explicit reason.
-				ELobbyMemberLeaveReason ResovedReason = ELobbyMemberLeaveReason::Disconnected;
-				if (const ELobbyMemberLeaveReason* MemberLeaveReason = LeaveReasons.Find(MemberData.Key))
-				{
-					ResovedReason = *MemberLeaveReason;
-				}
-				else
-				{
-					UE_LOG(LogTemp, Warning, TEXT("[FClientLobbyData::ApplyLobbyUpdateFromSnapshot] Member left lobby without giving reason: Lobby: %s, Member: %s"),
-						*ToLogString(PublicData->LobbyId), *ToLogString(MemberData.Key));
-				}
-
-				LeavingMembers.Emplace(MemberData.Value, ResovedReason);
-			}
-		}
-
-		// Remove member data and notify.
-		for (const TPair<TSharedRef<FClientLobbyMemberSnapshot>, ELobbyMemberLeaveReason>& MemberData : LeavingMembers)
-		{
-			MemberDataStorage.Remove(MemberData.Key->AccountId);
-			PublicData->Members.Remove(MemberData.Key->AccountId);
-
-			if (MemberData.Key->bIsLocalMember)
-			{
-				LocalMembers.Remove(MemberData.Key->AccountId);
-				Result.LeavingLocalMembers.Add(MemberData.Key->AccountId);
-			}
-
-			if (LobbyEvents)
-			{
-				LobbyEvents->OnLobbyMemberLeft.Broadcast(FLobbyMemberLeft{PublicData, MemberData.Key, MemberData.Value});
-			}
-		}
-	}
-
-	// Process member joins and attribute changes.
-	// Joining members are expected to have a member data snapshot.
-	for (TPair<FOnlineAccountIdHandle, TSharedRef<FClientLobbyMemberSnapshot>>& MemberSnapshot : LobbyMemberSnapshots)
-	{
-		if (LobbySnapshot.Members.Find(MemberSnapshot.Key) == nullptr)
-		{
-			UE_LOG(LogTemp, Warning, TEXT("[FClientLobbyData::ApplyLobbyUpdateFromSnapshot] Member update ignored for unknown member: Lobby: %s, Member: %s"),
-				*ToLogString(PublicData->LobbyId), *ToLogString(MemberSnapshot.Key));
-		}
-		else
-		{
-			// Check whether this member is already in the lobby.
-			if (TSharedRef<FClientLobbyMemberSnapshot>* MemberData = MemberDataStorage.Find(MemberSnapshot.Key))
-			{
-				// Member is already in the lobby. Check if their attributes changed.
-				TSet<FLobbyAttributeId> MemberAttributeChanges = ApplyAttributeUpdateFromSnapshot(MoveTemp(MemberSnapshot.Value->Attributes), (*MemberData)->Attributes);
-				if (LobbyEvents && !MemberAttributeChanges.IsEmpty())
-				{
-					LobbyEvents->OnLobbyMemberAttributesChanged.Broadcast(FLobbyMemberAttributesChanged{PublicData, (*MemberData), MemberAttributeChanges});
-				}
-			}
-			else
-			{
-				// Member is not in the lobby, add them.
-				MemberDataStorage.Add(MemberSnapshot.Key, MemberSnapshot.Value);
-				PublicData->Members.Add(MemberSnapshot.Key, MoveTemp(MemberSnapshot.Value));
-				if (LobbyEvents)
-				{
-					LobbyEvents->OnLobbyMemberJoined.Broadcast(FLobbyMemberJoined{PublicData, MemberSnapshot.Value});
-				}
-			}
-		}
-	}
-
-	// Process lobby state.
-	// The lobby owner notification must be processed after member updates as the new owner may not
-	// have existed in the previous snapshot.
-	if (PublicData->OwnerAccountId != LobbySnapshot.OwnerAccountId)
-	{
-		if (TSharedRef<FClientLobbyMemberSnapshot>* MemberData = MemberDataStorage.Find(LobbySnapshot.OwnerAccountId))
-		{
-			PublicData->OwnerAccountId = LobbySnapshot.OwnerAccountId;
-			if (LobbyEvents)
-			{
-				LobbyEvents->OnLobbyLeaderChanged.Broadcast(FLobbyLeaderChanged{PublicData, (*MemberData)});
-			}
-		}
-		else
-		{
-			UE_LOG(LogTemp, Warning, TEXT("[FClientLobbyData::ApplyLobbyUpdateFromSnapshot] Lobby owner chage failed - lobby member data not found: Lobby: %s, Member: %s"),
-				*ToLogString(PublicData->LobbyId), *ToLogString(LobbySnapshot.OwnerAccountId));
-		}
-	}
-
-	// Set remaining lobby data.
-	PublicData->MaxMembers = LobbySnapshot.MaxMembers;
-	PublicData->JoinPolicy = LobbySnapshot.JoinPolicy;
-
-	// If no local members remain in the lobby process lobby removal.
-	if (LobbyEvents && LocalMembers.IsEmpty())
-	{
-		TArray<TSharedRef<FClientLobbyMemberSnapshot>> LeavingMembers;
-		MemberDataStorage.GenerateValueArray(LeavingMembers);
-
-		for (TSharedRef<FClientLobbyMemberSnapshot>& LeavingMember : LeavingMembers)
-		{
-			MemberDataStorage.Remove(LeavingMember->AccountId);
-			PublicData->Members.Remove(LeavingMember->AccountId);
-			LobbyEvents->OnLobbyMemberLeft.Broadcast(FLobbyMemberLeft{PublicData, LeavingMember, ELobbyMemberLeaveReason::Left});
-		}
-
-		LobbyEvents->OnLobbyLeft.Broadcast(FLobbyLeft{PublicData});
-	}
-
-	return Result;
-}
-
-FApplyLobbyUpdateResult FClientLobbyData::ApplyLobbyUpdateFromLocalChanges(FClientLobbyDataChanges&& Changes, FLobbyEvents& LobbyEvents)
-{
-	FApplyLobbyUpdateResult Result;
-	const bool bIsActive = !LocalMembers.IsEmpty();
-
-	// Process lobby changes.
-	{
-		if (Changes.LocalName)
-		{
-			PublicData->LocalName = *Changes.LocalName;
-		}
-
-		// Join policy.
-		if (Changes.JoinPolicy)
-		{
-			PublicData->JoinPolicy = *Changes.JoinPolicy;
-		}
-
-		// Lobby owner.
-		if (Changes.OwnerAccountId && *Changes.OwnerAccountId != PublicData->OwnerAccountId)
-		{
-			if (TSharedRef<FClientLobbyMemberSnapshot>* MemberData = MemberDataStorage.Find(*Changes.OwnerAccountId))
-			{
-				PublicData->OwnerAccountId = *Changes.OwnerAccountId;
-				if (bIsActive)
-				{
-					LobbyEvents.OnLobbyLeaderChanged.Broadcast(FLobbyLeaderChanged{PublicData, (*MemberData)});
-				}
-			}
-			else
-			{
-				UE_LOG(LogTemp, Warning, TEXT("[FClientLobbyData::ApplyLobbyUpdateFromLocalChanges] Lobby owner chage failed - lobby member data not found: Lobby: %s, Member: %s"),
-					*ToLogString(PublicData->LobbyId), *ToLogString(*Changes.OwnerAccountId));
-			}
-		}
-
-		// Lobby schema.
-		if (Changes.LobbySchema && *Changes.LobbySchema != PublicData->SchemaName)
-		{
-			PublicData->SchemaName = *Changes.LobbySchema;
-			if (bIsActive)
-			{
-				LobbyEvents.OnLobbySchemaChanged.Broadcast(FLobbySchemaChanged{PublicData});
-			}
-		}
-
-		// Apply lobby attribute changes.
-		if (!Changes.MutatedAttributes.IsEmpty() || !Changes.ClearedAttributes.IsEmpty())
-		{
-			TSet<FLobbyAttributeId> ChangedLobbyAttributes = ApplyAttributeUpdateFromChanges(
-				MoveTemp(Changes.MutatedAttributes),
-				MoveTemp(Changes.ClearedAttributes),
-				PublicData->Attributes);
-
-			if (bIsActive && !ChangedLobbyAttributes.IsEmpty())
-			{
-				LobbyEvents.OnLobbyAttributesChanged.Broadcast(FLobbyAttributesChanged{PublicData, ChangedLobbyAttributes});
-			}
-		}
-	}
-
-	// Apply member changes
-	{
-		// Adding local members will make the lobby active. When the lobby moves from inactive to
-		// active join events must be generated for the lobby and each of its existing members.
-		if (!Changes.MutatedMembers.IsEmpty())
-		{
-			if (!bIsActive)
-			{
-				LobbyEvents.OnLobbyJoined.Broadcast(FLobbyJoined{PublicData});
-
-				for (TPair<FOnlineAccountIdHandle, TSharedRef<FClientLobbyMemberSnapshot>>& MemberData : MemberDataStorage)
-				{
-					// The lobby creator will already be in the list during lobby creation due to applying a lobby snapshot.
-					if (Changes.MutatedMembers.Find(MemberData.Key) == nullptr)
-					{
-						LobbyEvents.OnLobbyMemberJoined.Broadcast(FLobbyMemberJoined{PublicData, MemberData.Value});
-					}
-				}
-			}
-
-			// Create or modify state for each mutated member and notify.
-			// The lobby creator will already be in the list during lobby creation due to applying a lobby snapshot.
-			for (TPair<FOnlineAccountIdHandle, TSharedRef<FClientLobbyMemberDataChanges>>& MutatedMemberPair : Changes.MutatedMembers)
-			{
-				if (LocalMembers.Find(MutatedMemberPair.Key))
-				{
-					// Member data is expected to be added through changes before being mutated.
-					TSharedRef<FClientLobbyMemberSnapshot>* MemberDataPtr = MemberDataStorage.Find(MutatedMemberPair.Key);
-					if (ensure(MemberDataPtr))
-					{
-						TSharedRef<FClientLobbyMemberSnapshot>& MemberData = *MemberDataPtr;
-						TSet<FLobbyAttributeId> ChangedAttributes = ApplyAttributeUpdateFromChanges(
-							MoveTemp(MutatedMemberPair.Value->MutatedAttributes),
-							MoveTemp(MutatedMemberPair.Value->ClearedAttributes),
-							MemberData->Attributes);
-
-						if (!ChangedAttributes.IsEmpty())
-						{
-							LobbyEvents.OnLobbyMemberAttributesChanged.Broadcast(FLobbyMemberAttributesChanged{PublicData, MemberData, ChangedAttributes});
-						}
-					}
-				}
-				else
-				{
-					TSharedRef<FClientLobbyMemberSnapshot> NewMember = MakeShared<FClientLobbyMemberSnapshot>();
-					NewMember->bIsLocalMember = true;
-					NewMember->AccountId = MutatedMemberPair.Key;
-					NewMember->Attributes = MoveTemp(MutatedMemberPair.Value->MutatedAttributes);
-
-					// todo:
-					//NewMember->PlatformAccountId;
-					//NewMember->PlatformDisplayName;
-
-					MemberDataStorage.Add(NewMember->AccountId, NewMember);
-					PublicData->Members.Add(NewMember->AccountId, NewMember);
-					LocalMembers.Add(NewMember->AccountId);
-					LobbyEvents.OnLobbyMemberJoined.Broadcast(FLobbyMemberJoined{PublicData, NewMember});
-				}
-			}
-		}
-
-		// Handle leaving members.
-		// When the last local member leaves, leave events will be generated for each of the
-		// remaining lobby members and for the lobby itself.
-		if (!Changes.LeavingMembers.IsEmpty())
-		{
-			for (TPair<FOnlineAccountIdHandle, ELobbyMemberLeaveReason>& LeavingMemberPair : Changes.LeavingMembers)
-			{
-				// Member data is expected to be added through changes before being mutated.
-				TSharedRef<FClientLobbyMemberSnapshot>* MemberDataPtr = MemberDataStorage.Find(LeavingMemberPair.Key);
-				if (ensure(MemberDataPtr))
-				{
-					TSharedRef<FClientLobbyMemberSnapshot> MemberData = *MemberDataPtr;
-					if (MemberData->bIsLocalMember)
-					{
-						Result.LeavingLocalMembers.Add(LeavingMemberPair.Key);
-						LocalMembers.Remove(LeavingMemberPair.Key);
-					}
-
-					MemberDataStorage.Remove(LeavingMemberPair.Key);
-					PublicData->Members.Remove(LeavingMemberPair.Key);
-					LobbyEvents.OnLobbyMemberLeft.Broadcast(FLobbyMemberLeft{PublicData, MemberData, LeavingMemberPair.Value});
-				}
-			}
-
-			// If no local members remain in the lobby process lobby removal.
-			if (LocalMembers.IsEmpty())
-			{
-				TArray<TSharedRef<FClientLobbyMemberSnapshot>> LeavingMembers;
-				MemberDataStorage.GenerateValueArray(LeavingMembers);
-
-				for (TSharedRef<FClientLobbyMemberSnapshot>& LeavingMember : LeavingMembers)
-				{
-					MemberDataStorage.Remove(LeavingMember->AccountId);
-					PublicData->Members.Remove(LeavingMember->AccountId);
-					LobbyEvents.OnLobbyMemberLeft.Broadcast(FLobbyMemberLeft{PublicData, LeavingMember, ELobbyMemberLeaveReason::Left});
-				}
-
-				LobbyEvents.OnLobbyLeft.Broadcast(FLobbyLeft{PublicData});
-			}
-		}
-	}
-
-	return Result;
-}
-
-TSet<FLobbyAttributeId> FClientLobbyData::ApplyAttributeUpdateFromSnapshot(
-	TMap<FLobbyAttributeId, FLobbyVariant>&& AttributeSnapshot,
-	TMap<FLobbyAttributeId, FLobbyVariant>& ExistingAttributes)
-{
-	TSet<FLobbyAttributeId> Result;
-
-	// Clear the attributes.
-	{
-		// Determine which attributes were removed.
-		for (const TPair<FLobbyAttributeId, FLobbyVariant>& Attribute : ExistingAttributes)
-		{
-			if (!AttributeSnapshot.Contains(Attribute.Key))
-			{
-				Result.Add(Attribute.Key);
-			}
-		}
-
-		// Remove them.
-		for (const FLobbyAttributeId& AttributeId : Result)
-		{
-			ExistingAttributes.Remove(AttributeId);
-		}
-	}
-
-	// Apply attribute additions and mutations.
-	{
-		for (TPair<FLobbyAttributeId, FLobbyVariant>& Attribute : AttributeSnapshot)
-		{
-			if (FLobbyVariant* ExistingAttributeData = ExistingAttributes.Find(Attribute.Key))
-			{
-				if (*ExistingAttributeData != Attribute.Value)
-				{
-					Result.Add(MoveTemp(Attribute.Key));
-					*ExistingAttributeData = MoveTemp(Attribute.Value);
-				}
-			}
-			else
-			{
-				Result.Add(Attribute.Key);
-				ExistingAttributes.Add(MoveTemp(Attribute));
-			}
-		}
-	}
-
-	return Result;
-}
-
-TSet<FLobbyAttributeId> FClientLobbyData::ApplyAttributeUpdateFromChanges(
-	TMap<FLobbyAttributeId, FLobbyVariant>&& MutatedAttributes,
-	TSet<FLobbyAttributeId>&& ClearedAttributes,
-	TMap<FLobbyAttributeId, FLobbyVariant>& ExistingAttributes)
-{
-	TSet<FLobbyAttributeId> Result;
-
-	// Apply attribute additions and mutations.
-	{
-		for (TPair<FLobbyAttributeId, FLobbyVariant>& Attribute : MutatedAttributes)
-		{
-			if (FLobbyVariant* ExistingAttributeData = ExistingAttributes.Find(Attribute.Key))
-			{
-				if (*ExistingAttributeData != Attribute.Value)
-				{
-					Result.Add(MoveTemp(Attribute.Key));
-					*ExistingAttributeData = MoveTemp(Attribute.Value);
-				}
-			}
-			else
-			{
-				Result.Add(Attribute.Key);
-				ExistingAttributes.Add(MoveTemp(Attribute));
-			}
-		}
-	}
-
-	// Clear the attributes.
-	{
-		for (FLobbyAttributeId ClearedAttributeId : ClearedAttributes)
-		{
-			if (ExistingAttributes.Remove(ClearedAttributeId) > 0)
-			{
-				Result.Add(ClearedAttributeId);
-			}
-		}
-	}
-
-	return Result;
-}
-
 const EOS_HLobbyDetails FLobbyDetailsEOS::InvalidLobbyDetailsHandle = {};
 
 TDefaultErrorResultInternal<TSharedRef<FLobbyDetailsEOS>> FLobbyDetailsEOS::CreateFromLobbyId(
diff --git a/Engine/Plugins/Online/Experimental/OnlineServicesEOS/Source/Private/LobbiesEOSTypes.h b/Engine/Plugins/Online/Experimental/OnlineServicesEOS/Source/Private/LobbiesEOSTypes.h
--- a/Engine/Plugins/Online/Experimental/OnlineServicesEOS/Source/Private/LobbiesEOSTypes.h
+++ b/Engine/Plugins/Online/Experimental/OnlineServicesEOS/Source/Private/LobbiesEOSTypes.h
@@ -4,6 +4,7 @@
 
 #include "CoreMinimal.h"
 #include "Online/Lobbies.h"
+#include "Online/LobbiesCommonTypes.h"
 #include "OnlineIdEOS.h"
 #include "OnlineServicesEOSTypes.h"
 
@@ -48,48 +49,6 @@ inline ELobbyJoinPolicy TranslateJoinPolicy(EOS_ELobbyPermissionLevel JoinPolicy
 
 using FLobbySearchParameters = FFindLobbies::Params;
 
-/**
- * Local changes to a member to be applied to an existing lobby member snapshot.
- */
-struct FClientLobbyMemberDataChanges
-{
-	/** New or changed attributes. */
-	TMap<FLobbyAttributeId, FLobbyVariant> MutatedAttributes;
-
-	/** Attributes to be cleared. */
-	TSet<FLobbyAttributeId> ClearedAttributes;
-};
-
-/**
- * Local changes to a lobby to be applied to an existing lobby snapshot.
- */
-struct FClientLobbyDataChanges
-{
-	/** Local name for the lobby. */
-	TOptional<FName> LocalName;
-
-	/** Setting for new join policy. */
-	TOptional<ELobbyJoinPolicy> JoinPolicy;
-
-	/** Setting for lobby ownership change. */
-	TOptional<FOnlineAccountIdHandle> OwnerAccountId;
-
-	/** Setting for lobby schema change. */
-	TOptional<FLobbySchemaId> LobbySchema;
-
-	/** New or changed attributes. */
-	TMap<FLobbyAttributeId, FLobbyVariant> MutatedAttributes;
-
-	/** Attributes to be cleared. */
-	TSet<FLobbyAttributeId> ClearedAttributes;
-
-	/** Members to be added or changed. */
-	TMap<FOnlineAccountIdHandle, TSharedRef<FClientLobbyMemberDataChanges>> MutatedMembers;
-
-	/** Members to be removed. */
-	TMap<FOnlineAccountIdHandle, ELobbyMemberLeaveReason> LeavingMembers;
-};
-
 /**
 * Common components required to handle lobby requests.
 * To handle lifetime issues, only FLobbiesEOS should contain a strong reference to prerequisites.
@@ -103,89 +62,6 @@ struct FLobbyPrerequisitesEOS
 	TSharedRef<const FLobbySchema> ServiceSchema;
 };
 
-/**
-* Lobby snapshot data.
-* All contained data has been translated from EOS types.
-* Attributes have had their schema transformations applied.
-*/
-struct FClientLobbySnapshot
-{
-	FOnlineAccountIdHandle OwnerAccountId;
-	FName LocalName;
-	FName SchemaName;
-	int32 MaxMembers;
-	ELobbyJoinPolicy JoinPolicy;
-	TMap<FLobbyAttributeId, FLobbyVariant> Attributes;
-	TSet<FOnlineAccountIdHandle> Members;
-};
-
-struct FClientLobbyMemberSnapshot : public FLobbyMember
-{
-	bool bIsLocalMember = false;
-};
-
-struct FApplyLobbyUpdateResult
-{
-	TArray<FOnlineAccountIdHandle> LeavingLocalMembers;
-};
-
-/** Lobby data as seen by the client. */
-struct FClientLobbyData final
-{
-public:
-	FClientLobbyData(FOnlineLobbyIdHandle LobbyId);
-
-	TSharedRef<const FLobby> GetPublicDataPtr() const { return PublicData; }
-	const FLobby& GetPublicData() const { return *PublicData; }
-
-	TSharedPtr<const FClientLobbyMemberSnapshot> GetMemberData(FOnlineAccountIdHandle MemberAccountId) const;
-
-	/**
-	 * Apply updated lobby data and generate changes.
-	 * The LeaveReason provides context for members who have left since the most recent snapshot.
-	 * Changing a lobby generates events for the local client.
-	 */
-	FApplyLobbyUpdateResult ApplyLobbyUpdateFromServiceSnapshot(
-		FClientLobbySnapshot&& LobbySnapshot,
-		TMap<FOnlineAccountIdHandle, TSharedRef<FClientLobbyMemberSnapshot>>&& LobbyMemberSnapshots,
-		TMap<FOnlineAccountIdHandle, ELobbyMemberLeaveReason>&& LeaveReasons = TMap<FOnlineAccountIdHandle, ELobbyMemberLeaveReason>(),
-		FLobbyEvents* LobbyEvents = nullptr);
-
-	/**
-	 * Apply updated lobby data and generate changes.
-	 * Changing a lobby generates events for the local client.
-	 */
-	FApplyLobbyUpdateResult ApplyLobbyUpdateFromLocalChanges(FClientLobbyDataChanges&& Changes, FLobbyEvents& LobbyEvents);
-
-private:
-
-	/** Apply changes to a set of attributes. Returns the set of attribute IDs which changed. */
-	TSet<FLobbyAttributeId> ApplyAttributeUpdateFromSnapshot(
-		TMap<FLobbyAttributeId, FLobbyVariant>&& AttributeSnapshot,
-		TMap<FLobbyAttributeId, FLobbyVariant>& ExistingAttributes);
-
-	/** Apply changes to a set of attributes. Returns the set of attribute IDs which changed. */
-	TSet<FLobbyAttributeId> ApplyAttributeUpdateFromChanges(
-		TMap<FLobbyAttributeId, FLobbyVariant>&& MutatedAttributes,
-		TSet<FLobbyAttributeId>&& ClearedAttributes,
-		TMap<FLobbyAttributeId, FLobbyVariant>& ExistingAttributes);
-
-	/**
-	 * The shared pointer given back to user code with lobby operation results and notifications.
-	 * Any changes to this data are immediately available to users.
-	 */
-	TSharedRef<FLobby> PublicData;
-
-	/** Mutable lobby member data storage. */
-	TMap<FOnlineAccountIdHandle, TSharedRef<FClientLobbyMemberSnapshot>> MemberDataStorage;
-
-	/**
-	 * Keep track of which members are local to the client.
-	 * When all local members have been removed all members will be removed.
-	 */
-	TSet<FOnlineAccountIdHandle> LocalMembers;
-};
-
 /**
  * Lobby details are created for each user within the EOS lobby client. Certain operations such as
  * joining a lobby or applying an update to a lobby / member attributes require using the correct
