- Add ossv2 privilege check

diff --git a/Engine/Plugins/Online/Experimental/OnlineServices/Source/OnlineServicesInterface/Public/Online/Privileges.h b/Engine/Plugins/Online/Experimental/OnlineServices/Source/OnlineServicesInterface/Public/Online/Privileges.h
--- a/Engine/Plugins/Online/Experimental/OnlineServices/Source/OnlineServicesInterface/Public/Online/Privileges.h
+++ b/Engine/Plugins/Online/Experimental/OnlineServices/Source/OnlineServicesInterface/Public/Online/Privileges.h
@@ -80,6 +80,7 @@ struct FQueryUserPrivilege
 
 class IPrivileges
 {
+public:
 	/**
 		* Queries about a given privilege for a user
 		*/
