Forward Declare TSharedRef ToLogString override

diff --git a/Engine/Plugins/Online/Experimental/OnlineServices/Source/OnlineServicesInterface/Public/Online/OnlineUtils.h b/Engine/Plugins/Online/Experimental/OnlineServices/Source/OnlineServicesInterface/Public/Online/OnlineUtils.h
--- a/Engine/Plugins/Online/Experimental/OnlineServices/Source/OnlineServicesInterface/Public/Online/OnlineUtils.h
+++ b/Engine/Plugins/Online/Experimental/OnlineServices/Source/OnlineServicesInterface/Public/Online/OnlineUtils.h
@@ -9,6 +9,7 @@ namespace UE::Online {
 template <typename T> FString ToLogString(const TArray<T>& Array);
 template <typename K, typename V> FString ToLogString(const TMap<K, V>& Map);
 template <typename T, ESPMode Mode> FString ToLogString(const TSharedPtr<T, Mode>& Ptr);
+template <typename T, ESPMode Mode> FString ToLogString(const TSharedRef<T, Mode>& Ref);
 template <typename T> FString ToLogString(const TOptional<T> Optional);
 template <typename... Ts> FString ToLogString(const TVariant<Ts...>& Variant);
 template <typename T, ESPMode Mode> FString ToLogString(const TSharedRef<T, Mode>& Ref);
