Include account id handle in log output
Add help for online exec commands to output the required params

diff --git a/Engine/Plugins/Online/Experimental/OnlineServices/Source/OnlineServicesCommon/Public/Online/OnlineComponent.h b/Engine/Plugins/Online/Experimental/OnlineServices/Source/OnlineServicesCommon/Public/Online/OnlineComponent.h
--- a/Engine/Plugins/Online/Experimental/OnlineServices/Source/OnlineServicesCommon/Public/Online/OnlineComponent.h
+++ b/Engine/Plugins/Online/Experimental/OnlineServices/Source/OnlineServicesCommon/Public/Online/OnlineComponent.h
@@ -142,6 +142,20 @@ public:
 			{
 				return (*ExecHandler)->Exec(World, Cmd, Ar);
 			}
+			else if (Command == TEXT("HELP"))
+			{
+				return Help(World, Cmd, Ar);
+			}
+		}
+
+		return false;
+	}
+
+	bool Help(UWorld* World, const TCHAR* Cmd, FOutputDevice& Ar)
+	{
+		for (const TPair<FString, TUniquePtr<IOnlineExecHandler>>& Command : ExecCommands)
+		{
+			Command.Value->Help(World, Cmd, Ar);
 		}
 
 		return false;
diff --git a/Engine/Plugins/Online/Experimental/OnlineServices/Source/OnlineServicesCommon/Public/Online/OnlineExecHandler.h b/Engine/Plugins/Online/Experimental/OnlineServices/Source/OnlineServicesCommon/Public/Online/OnlineExecHandler.h
--- a/Engine/Plugins/Online/Experimental/OnlineServices/Source/OnlineServicesCommon/Public/Online/OnlineExecHandler.h
+++ b/Engine/Plugins/Online/Experimental/OnlineServices/Source/OnlineServicesCommon/Public/Online/OnlineExecHandler.h
@@ -19,6 +19,7 @@ class IOnlineExecHandler
 public:
 	virtual ~IOnlineExecHandler() {}
 	virtual bool Exec(UWorld* World, const TCHAR* Cmd, FOutputDevice& Ar) = 0;
+	virtual bool Help(UWorld* World, const TCHAR* Cmd, FOutputDevice& Ar) = 0;
 };
 
 namespace Private {
@@ -172,10 +173,10 @@ inline bool ParseOnlineExecParams(const TCHAR*& Cmd, uint64& Value)
 
 inline bool ParseOnlineExecParams(const TCHAR*& Cmd, FPlatformUserId& Value)
 {
-	int32 InternalId = -1;
-	if (ParseOnlineExecParams(Cmd, InternalId))
+	FString Token;
+	if (FParse::Token(Cmd, Token, true))
 	{
-		Value.CreateFromInternalId(InternalId);
+		Value = FPlatformMisc::GetPlatformUserForUserIndex(FCString::Strtoi(Cmd, nullptr, 10));
 	}
 	else
 	{
@@ -261,12 +262,20 @@ inline bool ParseOnlineExecParams(const TCHAR*& Cmd, TVariant<Ts...>& Variant)
 template<typename IdType>
 inline bool ParseOnlineExecParams(const TCHAR*& Cmd, TOnlineIdHandle<IdType>& Value)
 {
-	EOnlineServices Type;
-	uint32 Handle;
-	if (ParseOnlineExecParams(Cmd, Type) && ParseOnlineExecParams(Cmd, Handle))
+	FString Token;
+	if (FParse::Token(Cmd, Token, true))
 	{
-		Value = TOnlineIdHandle<IdType>(Type, Handle);
-		return true;
+		FString ServicesString;
+		FString HandleString;
+
+		if (Token.Split(TEXT(":"), &ServicesString, &HandleString, ESearchCase::CaseSensitive))
+		{
+			EOnlineServices Services;
+			LexFromString(Services, *ServicesString);
+			uint32 Handle = static_cast<uint32>(FCString::Strtoui64(*HandleString, nullptr, 10));
+			Value = TOnlineIdHandle<IdType>(Services, Handle);
+			return true;
+		}
 	}
 
 	return false;
@@ -315,7 +324,7 @@ public:
 	{
 	}
 
-	virtual bool Exec(UWorld* World, const TCHAR* Cmd, FOutputDevice& Ar)
+	virtual bool Exec(UWorld* World, const TCHAR* Cmd, FOutputDevice& Ar) override
 	{
 		typename OpType::Params Params;
 		if (!Private::ParseOnlineExecParams(Cmd, Params))
@@ -341,6 +350,20 @@ public:
 		return true;
 	}
 
+	virtual bool Help(UWorld* World, const TCHAR* Cmd, FOutputDevice& Ar) override
+	{
+		FString HelpString = OpType::Name;
+		Meta::VisitFields<typename OpType::Params>([&HelpString](const auto& Field)
+			{
+				HelpString += TEXT(" ");
+				HelpString += Field.Name;
+			});
+
+		Ar.Log(HelpString);
+		
+		return true;
+	}
+
 private:
 	InterfaceType* Interface;
 	MemberFunctionPtrType Function;
@@ -360,6 +383,11 @@ public:
 		return Component->Exec(World, Cmd, Ar);
 	}
 
+	virtual bool Help(UWorld* World, const TCHAR* Cmd, FOutputDevice& Ar) override
+	{
+		return Component->Help(World, Cmd, Ar);
+	}
+
 private:
 	T* Component;
 };
diff --git a/Engine/Source/Runtime/CoreOnline/Private/Online/CoreOnline.cpp b/Engine/Source/Runtime/CoreOnline/Private/Online/CoreOnline.cpp
--- a/Engine/Source/Runtime/CoreOnline/Private/Online/CoreOnline.cpp
+++ b/Engine/Source/Runtime/CoreOnline/Private/Online/CoreOnline.cpp
@@ -291,7 +291,7 @@ IOnlineAccountIdRegistry* FOnlineIdRegistryRegistry::GetAccountIdRegistry(EOnlin
 
 FString ToLogString(const FOnlineAccountIdHandle& Id)
 {
-	return FOnlineIdRegistryRegistry::Get().ToLogString(Id);
+	return FString::Printf(TEXT("%s:%d (%s)"), LexToString(Id.GetOnlineServicesType()), Id.GetHandle(), *FOnlineIdRegistryRegistry::Get().ToLogString(Id));
 }
 
 FString ToLogString(const FOnlineLobbyIdHandle& Id)
