OnlineAsyncOp steps now work with non-default constructible types

diff --git a/Engine/Plugins/Online/Experimental/OnlineServices/Source/OnlineServicesCommon/Public/Online/OnlineAsyncOp.h b/Engine/Plugins/Online/Experimental/OnlineServices/Source/OnlineServicesCommon/Public/Online/OnlineAsyncOp.h
--- a/Engine/Plugins/Online/Experimental/OnlineServices/Source/OnlineServicesCommon/Public/Online/OnlineAsyncOp.h
+++ b/Engine/Plugins/Online/Experimental/OnlineServices/Source/OnlineServicesCommon/Public/Online/OnlineAsyncOp.h
@@ -303,6 +303,14 @@ public:
 	{
 	}
 
+	~TStep()
+	{
+		if (bResultSet)
+		{
+			DestructItem(Result.GetTypedPtr());
+		}
+	}
+
 	template <typename OpType, typename LastResultType, typename CallableType>
 	void SetExecFunction(TOnlineAsyncOp<OpType>& InOperation, LastResultType&& InLastResult, CallableType&& InCallable)
 	{
@@ -321,7 +329,7 @@ public:
 							TSharedPtr<TOnlineAsyncOp<OpType>> PinnedOperation2 = WeakOperation.Pin();
 							if (PinnedOperation2)
 							{
-								Result = Value;
+								EmplaceResult(Value);
 								PinnedOperation2->ExecuteNextStep();
 							}
 						});
@@ -336,14 +344,14 @@ public:
 							TSharedPtr<TOnlineAsyncOp<OpType>> PinnedOperation2 = WeakOperation.Pin();
 							if (PinnedOperation2)
 							{
-								Result = Value;
+								EmplaceResult(Value);
 								PinnedOperation2->ExecuteNextStep();
 							}
 						});
 				}
 				else
 				{
-					Result = Callable(*PinnedOperation, MoveTempIfPossible(LastResult));
+					EmplaceResult(Callable(*PinnedOperation, MoveTempIfPossible(LastResult)));
 					PinnedOperation->ExecuteNextStep();
 				}
 			}
@@ -368,7 +376,7 @@ public:
 							TSharedPtr<TOnlineAsyncOp<OpType>> PinnedOperation2 = WeakOperation.Pin();
 							if (PinnedOperation2)
 							{
-								Result = Value;
+								EmplaceResult(Value);
 								PinnedOperation2->ExecuteNextStep();
 							}
 						});
@@ -383,14 +391,14 @@ public:
 							TSharedPtr<TOnlineAsyncOp<OpType>> PinnedOperation2 = WeakOperation.Pin();
 							if (PinnedOperation2)
 							{
-								Result = Value;
+								EmplaceResult(Value);
 								PinnedOperation2->ExecuteNextStep();
 							}
 						});
 				}
 				else
 				{
-					Result = Callable(*PinnedOperation);
+					EmplaceResult(Callable(*PinnedOperation));
 					PinnedOperation->ExecuteNextStep();
 				}
 			}
@@ -409,15 +417,24 @@ public:
 	}
 
 
-	ResultType* GetResultPtr()
+	ResultType& GetResultRef()
 	{
-		return &Result;
+		return *Result.GetTypedPtr();
 	}
 
 private:
+	template<typename... ArgTypes>
+	void EmplaceResult(ArgTypes&&... Args)
+	{
+		check(!bResultSet)
+		new(Result.GetTypedPtr()) ResultType(Forward<ArgTypes>(Args)...);
+		bResultSet = true;
+	}
+
 	FOnlineAsyncExecutionPolicy ExecutionPolicy;
 	TUniqueFunction<void()> ExecFunction;
-	ResultType Result;
+	TTypeCompatibleBytes<ResultType> Result;
+	bool bResultSet = false;
 };
 
 
@@ -509,7 +526,7 @@ template <typename Outer, typename OpType, typename LastResultType>
 class TOnlineAsyncOpBase
 {
 public:
-	TOnlineAsyncOpBase(LastResultType* InLastResult)
+	TOnlineAsyncOpBase(LastResultType& InLastResult)
 		: LastResult(InLastResult)
 	{
 	}
@@ -528,7 +545,7 @@ public:
 	auto Then(CallableType&& Callable, FOnlineAsyncExecutionPolicy ExecutionPolicy = FOnlineAsyncExecutionPolicy::RunOnGameThread());
 
 protected:
-	LastResultType* LastResult;
+	LastResultType& LastResult;
 };
 
 template <typename Outer, typename OpType>
@@ -558,12 +575,45 @@ class TOnlineChainableAsyncOp : public Private::TOnlineAsyncOpBase<TOnlineChaina
 public:
 	using Super = Private::TOnlineAsyncOpBase<TOnlineChainableAsyncOp<OpType, T>, OpType, T>;
 
-	TOnlineChainableAsyncOp(TOnlineAsyncOp<OpType>& InOwningOperation, T* InLastResult)
+	TOnlineChainableAsyncOp(TOnlineAsyncOp<OpType>& InOwningOperation, std::enable_if_t<!std::is_same_v<T, void>, T>& InLastResult)
 		: Super(InLastResult)
 		, OwningOperation(InOwningOperation)
 	{
 	}
 
+	TOnlineChainableAsyncOp(TOnlineChainableAsyncOp&& Other)
+		: OwningOperation(Other.OwningOperation)
+	{
+	}
+
+	TOnlineChainableAsyncOp& operator=(TOnlineChainableAsyncOp&& Other)
+	{
+		check(&OwningOperation == &Other.OwningOperation); // Can't reassign this
+		Super::operator=(MoveTemp(Other));
+		return *this;
+	}
+
+	template <typename QueueType>
+	void Enqueue(QueueType& Queue)
+	{
+		static_assert(std::is_same_v<T, void>, "Continuation result discarded. Continuation prior to calling Enqueue must have a void or TFuture<void> return type.");
+		OwningOperation.Enqueue(Queue);
+	}
+
+	TOnlineAsyncOp<OpType>& GetOwningOperation()
+	{
+		return OwningOperation;
+	}
+
+protected:
+	TOnlineAsyncOp<OpType>& OwningOperation;
+};
+
+template <typename OpType>
+class TOnlineChainableAsyncOp<OpType, void> : public Private::TOnlineAsyncOpBase<TOnlineChainableAsyncOp<OpType, void>, OpType, void>
+{
+public:
+	using Super = Private::TOnlineAsyncOpBase<TOnlineChainableAsyncOp<OpType, void>, OpType, void>;
 
 	TOnlineChainableAsyncOp(TOnlineAsyncOp<OpType>& InOwningOperation)
 		: Super()
@@ -586,7 +636,6 @@ public:
 	template <typename QueueType>
 	void Enqueue(QueueType& Queue)
 	{
-		static_assert(std::is_same_v<T, void>, "Continuation result discarded. Continuation prior to calling Enqueue must have a void or TFuture<void> return type.");
 		OwningOperation.Enqueue(Queue);
 	}
 
@@ -920,7 +969,7 @@ auto TOnlineAsyncOpBase<Outer, OpType, LastResultType>::Then(CallableType&& InCa
 
 	TStep<ResultType>* Step = new TStep<ResultType>(MoveTemp(ExecutionPolicy));
 	TUniquePtr<IStep> StepPtr(Step);
-	Step->SetExecFunction(Op, *LastResult, MoveTemp(InCallable));
+	Step->SetExecFunction(Op, LastResult, MoveTemp(InCallable));
 
 	Op.AddStep(MoveTemp(StepPtr));
 
@@ -930,7 +979,7 @@ auto TOnlineAsyncOpBase<Outer, OpType, LastResultType>::Then(CallableType&& InCa
 	}
 	else
 	{
-		return TOnlineChainableAsyncOp<OpType, ResultType>(Op, Step->GetResultPtr());
+		return TOnlineChainableAsyncOp<OpType, ResultType>(Op, Step->GetResultRef());
 	}
 }
 
@@ -954,7 +1003,7 @@ auto TOnlineAsyncOpBase<Outer, OpType, void>::Then(CallableType&& InCallable, FO
 	}
 	else
 	{
-		return TOnlineChainableAsyncOp<OpType, ResultType>(Op, Step->GetResultPtr());
+		return TOnlineChainableAsyncOp<OpType, ResultType>(Op, Step->GetResultRef());
 	}
 }
 
