Fix issue with UE_ONLINE_ERROR and UE_ONLINE_ERROR_COMMON macros populating Category where the Value should be stored and Value where Category should be stored
- Adds Test to ErrorsTests for OI-3725 regression testing

diff --git a/Engine/Plugins/Online/Experimental/OnlineServices/Source/OnlineServicesInterface/Public/Online/OnlineErrorCode.h b/Engine/Plugins/Online/Experimental/OnlineServices/Source/OnlineServicesInterface/Public/Online/OnlineErrorCode.h
--- a/Engine/Plugins/Online/Experimental/OnlineServices/Source/OnlineServicesInterface/Public/Online/OnlineErrorCode.h
+++ b/Engine/Plugins/Online/Experimental/OnlineServices/Source/OnlineServicesInterface/Public/Online/OnlineErrorCode.h
@@ -78,14 +78,14 @@ namespace ErrorCode { namespace Category { \
 		}
 
 #define UE_ONLINE_ERROR(CategoryName, Name, ErrorCodeValue, ErrorMessage, ErrorText) \
-    namespace ErrorCode { namespace CategoryName { static constexpr ErrorCodeType Name = Create(ErrorCodeValue, Category::CategoryName); } } \
+    namespace ErrorCode { namespace CategoryName { static constexpr ErrorCodeType Name = Create(Category::CategoryName, ErrorCodeValue); } } \
 	namespace CategoryName { \
 		UE_ONLINE_ERROR_INTERNAL(CategoryName, Name, ErrorCodeValue, ErrorMessage, ErrorText)\
 	}
 
 // Same as above except no namespace around the error, e.g. Errors::NotImplemented instead of Errors::OnlineServices::NotImplemented
 #define UE_ONLINE_ERROR_COMMON(CategoryName, Name, ErrorCodeValue, ErrorMessage, ErrorText) \
-    namespace ErrorCode { namespace CategoryName { static constexpr ErrorCodeType Name = Create(ErrorCodeValue, Category::CategoryName); } } \
+    namespace ErrorCode { namespace CategoryName { static constexpr ErrorCodeType Name = Create(Category::CategoryName, ErrorCodeValue); } } \
 	UE_ONLINE_ERROR_INTERNAL(CategoryName, Name, ErrorCodeValue, ErrorMessage, ErrorText)
 
 } /* namespace UE::Online::Errors */
\ No newline at end of file
