Add GetAuthToken to UE::Online::IAuth

diff --git a/Engine/Plugins/Online/Experimental/OnlineServices/Source/OnlineServicesCommon/Private/Online/AuthCommon.cpp b/Engine/Plugins/Online/Experimental/OnlineServices/Source/OnlineServicesCommon/Private/Online/AuthCommon.cpp
--- a/Engine/Plugins/Online/Experimental/OnlineServices/Source/OnlineServicesCommon/Private/Online/AuthCommon.cpp
+++ b/Engine/Plugins/Online/Experimental/OnlineServices/Source/OnlineServicesCommon/Private/Online/AuthCommon.cpp
@@ -18,6 +18,7 @@ void FAuthCommon::RegisterCommands()
 	RegisterCommand(&FAuthCommon::Login);
 	RegisterCommand(&FAuthCommon::Logout);
 	RegisterCommand(&FAuthCommon::GenerateAuth);
+	RegisterCommand(&FAuthCommon::GetAuthToken);
 	RegisterCommand(&FAuthCommon::GetAccountByPlatformUserId);
 	RegisterCommand(&FAuthCommon::GetAccountByAccountId);
 }
@@ -43,6 +44,11 @@ TOnlineAsyncOpHandle<FAuthGenerateAuth> FAuthCommon::GenerateAuth(FAuthGenerateA
 	return Operation->GetHandle();
 }
 
+TOnlineResult<FAuthGetAuthToken> FAuthCommon::GetAuthToken(FAuthGetAuthToken::Params&& Params)
+{
+	return TOnlineResult<FAuthGetAuthToken>(Errors::NotImplemented());
+}
+
 TOnlineResult<FAuthGetAccountByPlatformUserId> FAuthCommon::GetAccountByPlatformUserId(FAuthGetAccountByPlatformUserId::Params&& Params)
 {
 	return TOnlineResult<FAuthGetAccountByPlatformUserId>(Errors::NotImplemented());
diff --git a/Engine/Plugins/Online/Experimental/OnlineServices/Source/OnlineServicesCommon/Public/Online/AuthCommon.h b/Engine/Plugins/Online/Experimental/OnlineServices/Source/OnlineServicesCommon/Public/Online/AuthCommon.h
--- a/Engine/Plugins/Online/Experimental/OnlineServices/Source/OnlineServicesCommon/Public/Online/AuthCommon.h
+++ b/Engine/Plugins/Online/Experimental/OnlineServices/Source/OnlineServicesCommon/Public/Online/AuthCommon.h
@@ -23,6 +23,7 @@ public:
 	virtual TOnlineAsyncOpHandle<FAuthLogin> Login(FAuthLogin::Params&& Params) override;
 	virtual TOnlineAsyncOpHandle<FAuthLogout> Logout(FAuthLogout::Params&& Params) override;
 	virtual TOnlineAsyncOpHandle<FAuthGenerateAuth> GenerateAuth(FAuthGenerateAuth::Params&& Params) override;
+	virtual TOnlineResult<FAuthGetAuthToken> GetAuthToken(FAuthGetAuthToken::Params&& Params) override;
 	virtual TOnlineResult<FAuthGetAccountByPlatformUserId> GetAccountByPlatformUserId(FAuthGetAccountByPlatformUserId::Params&& Params) override;
 	virtual TOnlineResult<FAuthGetAccountByAccountId> GetAccountByAccountId(FAuthGetAccountByAccountId::Params&& Params) override;
 	virtual TOnlineEvent<void(const FLoginStatusChanged&)> OnLoginStatusChanged() override;
diff --git a/Engine/Plugins/Online/Experimental/OnlineServices/Source/OnlineServicesInterface/Public/Online/Auth.h b/Engine/Plugins/Online/Experimental/OnlineServices/Source/OnlineServicesInterface/Public/Online/Auth.h
--- a/Engine/Plugins/Online/Experimental/OnlineServices/Source/OnlineServicesInterface/Public/Online/Auth.h
+++ b/Engine/Plugins/Online/Experimental/OnlineServices/Source/OnlineServicesInterface/Public/Online/Auth.h
@@ -91,6 +91,23 @@ struct FAuthGenerateAuth
 	};
 };
 
+struct FAuthGetAuthToken
+{
+	static constexpr TCHAR Name[] = TEXT("GetAuthToken");
+
+	struct Params
+	{
+		FOnlineAccountIdHandle LocalUserId;
+		FString Type;
+		TArray<FString> Scopes;
+	};
+
+	struct Result
+	{
+		FString Token;
+	};
+};
+
 struct FAuthGetAccountByPlatformUserId
 {
 	static constexpr TCHAR Name[] = TEXT("GetAccountByPlatformUserId");
@@ -150,6 +167,11 @@ public:
 	 */
 	virtual TOnlineAsyncOpHandle<FAuthGenerateAuth> GenerateAuth(FAuthGenerateAuth::Params&& Params) = 0;
 
+	/**
+	 * Retrive a cached auth token
+	 */
+	virtual TOnlineResult<FAuthGetAuthToken> GetAuthToken(FAuthGetAuthToken::Params&& Params) = 0;
+
 	/**
 	 * Get logged in user by local user num
 	 */ 
@@ -204,6 +226,16 @@ END_ONLINE_STRUCT_META()
 BEGIN_ONLINE_STRUCT_META(FAuthGenerateAuth::Result)
 END_ONLINE_STRUCT_META()
 
+BEGIN_ONLINE_STRUCT_META(FAuthGetAuthToken::Params)
+	ONLINE_STRUCT_FIELD(FAuthGetAuthToken::Params, LocalUserId),
+	ONLINE_STRUCT_FIELD(FAuthGetAuthToken::Params, Type),
+	ONLINE_STRUCT_FIELD(FAuthGetAuthToken::Params, Scopes)
+END_ONLINE_STRUCT_META()
+
+BEGIN_ONLINE_STRUCT_META(FAuthGetAuthToken::Result)
+	ONLINE_STRUCT_FIELD(FAuthGetAuthToken::Result, Token)
+END_ONLINE_STRUCT_META()
+
 BEGIN_ONLINE_STRUCT_META(FAuthGetAccountByPlatformUserId::Params)
 	ONLINE_STRUCT_FIELD(FAuthGetAccountByPlatformUserId::Params, PlatformUserId)
 END_ONLINE_STRUCT_META()
