OSSv2 Sessions Interface - Common Implementation

diff --git a/Engine/Plugins/Online/Experimental/OnlineServices/Source/OnlineServicesCommon/Private/Online/OnlineServicesCommon.cpp b/Engine/Plugins/Online/Experimental/OnlineServices/Source/OnlineServicesCommon/Private/Online/OnlineServicesCommon.cpp
--- a/Engine/Plugins/Online/Experimental/OnlineServices/Source/OnlineServicesCommon/Private/Online/OnlineServicesCommon.cpp
+++ b/Engine/Plugins/Online/Experimental/OnlineServices/Source/OnlineServicesCommon/Private/Online/OnlineServicesCommon.cpp
@@ -65,6 +65,11 @@
 	return ILobbiesPtr(AsShared(), Get<ILobbies>());
 }
 
+ISessionsPtr FOnlineServicesCommon::GetSessionsInterface()
+{
+	return ISessionsPtr(AsShared(), Get<ISessions>());
+}
+
 IConnectivityPtr FOnlineServicesCommon::GetConnectivityInterface()
 {
 	return IConnectivityPtr(AsShared(), Get<IConnectivity>());
diff --git a/Engine/Plugins/Online/Experimental/OnlineServices/Source/OnlineServicesCommon/Private/Online/SessionsCommon.cpp b/Engine/Plugins/Online/Experimental/OnlineServices/Source/OnlineServicesCommon/Private/Online/SessionsCommon.cpp
new file mode 100644
--- /dev/null
+++ b/Engine/Plugins/Online/Experimental/OnlineServices/Source/OnlineServicesCommon/Private/Online/SessionsCommon.cpp
@@ -0,0 +1,157 @@
+// Copyright Epic Games, Inc. All Rights Reserved.
+
+#include "Online/SessionsCommon.h"
+
+#include "Online/Auth.h"
+#include "Online/OnlineServicesCommon.h"
+
+namespace UE::Online {
+
+	FSessionsCommon::FSessionsCommon(FOnlineServicesCommon& InServices)
+		: TOnlineComponent(TEXT("Sessions"), InServices)
+	{
+	}
+
+	void FSessionsCommon::Initialize()
+	{
+		TOnlineComponent<ISessions>::Initialize();
+	}
+
+	void FSessionsCommon::RegisterCommands()
+	{
+		TOnlineComponent<ISessions>::RegisterCommands();
+
+		RegisterCommand(&FSessionsCommon::GetAllSessions);
+		RegisterCommand(&FSessionsCommon::GetSessionByName);
+		RegisterCommand(&FSessionsCommon::GetSessionById);
+		RegisterCommand(&FSessionsCommon::CreateSession);
+		RegisterCommand(&FSessionsCommon::UpdateSession);
+		RegisterCommand(&FSessionsCommon::LeaveSession);
+		RegisterCommand(&FSessionsCommon::FindSessions);
+		RegisterCommand(&FSessionsCommon::StartMatchmaking);
+		RegisterCommand(&FSessionsCommon::JoinSession);
+		RegisterCommand(&FSessionsCommon::SendSessionInvite);
+		RegisterCommand(&FSessionsCommon::QuerySessionInvites);
+		RegisterCommand(&FSessionsCommon::RejectSessionInvite);
+		RegisterCommand(&FSessionsCommon::RegisterPlayers);
+		RegisterCommand(&FSessionsCommon::UnregisterPlayers);
+	}
+
+	TOnlineResult<FGetAllSessions> FSessionsCommon::GetAllSessions(FGetAllSessions::Params&& Params)
+	{
+		return TOnlineResult<FGetAllSessions>(Errors::NotImplemented());
+	}
+
+	TOnlineResult<FGetSessionByName> FSessionsCommon::GetSessionByName(FGetSessionByName::Params&& Params)
+	{
+		return TOnlineResult<FGetSessionByName>(Errors::NotImplemented());
+	}
+
+	TOnlineResult<FGetSessionById> FSessionsCommon::GetSessionById(FGetSessionById::Params&& Params)
+	{
+		return TOnlineResult<FGetSessionById>(Errors::NotImplemented());
+	}
+
+	TOnlineAsyncOpHandle<FCreateSession> FSessionsCommon::CreateSession(FCreateSession::Params&& Params)
+	{
+		TOnlineAsyncOpRef<FCreateSession> Operation = GetOp<FCreateSession>(MoveTemp(Params));
+		Operation->SetError(Errors::NotImplemented());
+		return Operation->GetHandle();
+	}
+
+	TOnlineAsyncOpHandle<FUpdateSession> FSessionsCommon::UpdateSession(FUpdateSession::Params&& Params)
+	{
+		TOnlineAsyncOpRef<FUpdateSession> Operation = GetOp<FUpdateSession>(MoveTemp(Params));
+		Operation->SetError(Errors::NotImplemented());
+		return Operation->GetHandle();
+	}
+
+	TOnlineAsyncOpHandle<FLeaveSession> FSessionsCommon::LeaveSession(FLeaveSession::Params&& Params)
+	{
+		TOnlineAsyncOpRef<FLeaveSession> Operation = GetOp<FLeaveSession>(MoveTemp(Params));
+		Operation->SetError(Errors::NotImplemented());
+		return Operation->GetHandle();
+	}
+
+	TOnlineAsyncOpHandle<FFindSessions> FSessionsCommon::FindSessions(FFindSessions::Params&& Params)
+	{
+		TOnlineAsyncOpRef<FFindSessions> Operation = GetOp<FFindSessions>(MoveTemp(Params));
+		Operation->SetError(Errors::NotImplemented());
+		return Operation->GetHandle();
+	}
+
+	TOnlineAsyncOpHandle<FStartMatchmaking> FSessionsCommon::StartMatchmaking(FStartMatchmaking::Params&& Params)
+	{
+		TOnlineAsyncOpRef<FStartMatchmaking> Operation = GetOp<FStartMatchmaking>(MoveTemp(Params));
+		Operation->SetError(Errors::NotImplemented());
+		return Operation->GetHandle();
+	}
+
+	TOnlineAsyncOpHandle<FJoinSession> FSessionsCommon::JoinSession(FJoinSession::Params&& Params)
+	{
+		TOnlineAsyncOpRef<FJoinSession> Operation = GetOp<FJoinSession>(MoveTemp(Params));
+		Operation->SetError(Errors::NotImplemented());
+		return Operation->GetHandle();
+	}
+
+	TOnlineAsyncOpHandle<FSendSessionInvite> FSessionsCommon::SendSessionInvite(FSendSessionInvite::Params&& Params)
+	{
+		TOnlineAsyncOpRef<FSendSessionInvite> Operation = GetOp<FSendSessionInvite>(MoveTemp(Params));
+		Operation->SetError(Errors::NotImplemented());
+		return Operation->GetHandle();
+	}
+
+	TOnlineAsyncOpHandle<FQuerySessionInvites> FSessionsCommon::QuerySessionInvites(FQuerySessionInvites::Params&& Params)
+	{
+		TOnlineAsyncOpRef<FQuerySessionInvites> Operation = GetOp<FQuerySessionInvites>(MoveTemp(Params));
+		Operation->SetError(Errors::NotImplemented());
+		return Operation->GetHandle();
+	}
+
+	TOnlineAsyncOpHandle<FRejectSessionInvite> FSessionsCommon::RejectSessionInvite(FRejectSessionInvite::Params&& Params)
+	{
+		TOnlineAsyncOpRef<FRejectSessionInvite> Operation = GetOp<FRejectSessionInvite>(MoveTemp(Params));
+		Operation->SetError(Errors::NotImplemented());
+		return Operation->GetHandle();
+	}
+
+	TOnlineAsyncOpHandle<FRegisterPlayers> FSessionsCommon::RegisterPlayers(FRegisterPlayers::Params&& Params)
+	{
+		TOnlineAsyncOpRef<FRegisterPlayers> Operation = GetOp<FRegisterPlayers>(MoveTemp(Params));
+		Operation->SetError(Errors::NotImplemented());
+		return Operation->GetHandle();
+	}
+
+	TOnlineAsyncOpHandle<FUnregisterPlayers> FSessionsCommon::UnregisterPlayers(FUnregisterPlayers::Params&& Params)
+	{
+		TOnlineAsyncOpRef<FUnregisterPlayers> Operation = GetOp<FUnregisterPlayers>(MoveTemp(Params));
+		Operation->SetError(Errors::NotImplemented());
+		return Operation->GetHandle();
+	}
+
+	TOnlineEvent<void(const FSessionJoined&)> FSessionsCommon::OnSessionJoined()
+	{
+		return SessionEvents.OnSessionJoined;
+	}
+
+	TOnlineEvent<void(const FSessionLeft&)> FSessionsCommon::OnSessionLeft()
+	{
+		return SessionEvents.OnSessionLeft;
+	}
+
+	TOnlineEvent<void(const FSessionUpdated&)> FSessionsCommon::OnSessionUpdated()
+	{
+		return SessionEvents.OnSessionUpdated;
+	}
+
+	TOnlineEvent<void(const FInviteReceived&)> FSessionsCommon::OnInviteReceived()
+	{
+		return SessionEvents.OnInviteReceived;
+	}
+
+	TOnlineEvent<void(const FInviteAccepted&)> FSessionsCommon::OnInviteAccepted()
+	{
+		return SessionEvents.OnInviteAccepted;
+	}
+
+/* UE::Online */ }
\ No newline at end of file
diff --git a/Engine/Plugins/Online/Experimental/OnlineServices/Source/OnlineServicesCommon/Public/Online/OnlineExecHandler.h b/Engine/Plugins/Online/Experimental/OnlineServices/Source/OnlineServicesCommon/Public/Online/OnlineExecHandler.h
--- a/Engine/Plugins/Online/Experimental/OnlineServices/Source/OnlineServicesCommon/Public/Online/OnlineExecHandler.h
+++ b/Engine/Plugins/Online/Experimental/OnlineServices/Source/OnlineServicesCommon/Public/Online/OnlineExecHandler.h
@@ -67,8 +67,12 @@
 inline bool ParseOnlineExecParams(const TCHAR*& Cmd, uint32& Value);
 inline bool ParseOnlineExecParams(const TCHAR*& Cmd, int64& Value);
 inline bool ParseOnlineExecParams(const TCHAR*& Cmd, uint64& Value);
+inline bool ParseOnlineExecParams(const TCHAR*& Cmd, bool& Value);
 template <typename T> inline bool ParseOnlineExecParams(const TCHAR*& Cmd, TArray<T>& Array);
 template <typename T, typename U> inline bool ParseOnlineExecParams(const TCHAR*& Cmd, TMap<T, U>& Map);
+template <typename T> inline bool ParseOnlineExecParams(const TCHAR*& Cmd, TSharedPtr<T>& Ptr);
+template <typename T> inline bool ParseOnlineExecParams(const TCHAR*& Cmd, TSharedRef<T>& Ref);
+template <typename T> inline bool ParseOnlineExecParams(const TCHAR*& Cmd, TOptional<T>& Optional);
 template <typename... Ts> inline bool ParseOnlineExecParams(const TCHAR*& Cmd, TVariant<Ts...>& Variant);
 template <typename IdType> inline bool ParseOnlineExecParams(const TCHAR*& Cmd, TOnlineIdHandle<IdType>& Value);
 template <typename T> std::enable_if_t<!TModels<Meta::COnlineMetadataAvailable, T>::Value, bool> ParseOnlineExecParams(const TCHAR*& Cmd, T& Value);
@@ -171,6 +175,21 @@
 	return true;
 }
 
+inline bool ParseOnlineExecParams(const TCHAR*& Cmd, bool& Value)
+{
+	FString Token;
+	if (FParse::Token(Cmd, Token, true))
+	{
+		Value = FCString::ToBool(*Token);
+	}
+	else
+	{
+		return false;
+	}
+
+	return true;
+}
+
 inline bool ParseOnlineExecParams(const TCHAR*& Cmd, FPlatformUserId& Value)
 {
 	FString Token;
@@ -252,6 +271,28 @@
 	return true;
 }
 
+
+template <typename T>
+inline bool ParseOnlineExecParams(const TCHAR*& Cmd, TSharedPtr<T>& Ptr)
+{
+	// TODO: This is a temporary stub implementation
+	return true;
+}
+
+template <typename T>
+inline bool ParseOnlineExecParams(const TCHAR*& Cmd, TSharedRef<T>& Ref)
+{
+	// TODO: This is a temporary stub implementation
+	return true;
+}
+
+template <typename T>
+inline bool ParseOnlineExecParams(const TCHAR*& Cmd, TOptional<T>& Optional)
+{
+	// TODO: This is a temporary stub implementation
+	return true;
+}
+
 template <typename... Ts>
 inline bool ParseOnlineExecParams(const TCHAR*& Cmd, TVariant<Ts...>& Variant)
 {
diff --git a/Engine/Plugins/Online/Experimental/OnlineServices/Source/OnlineServicesCommon/Public/Online/OnlineServicesCommon.h b/Engine/Plugins/Online/Experimental/OnlineServices/Source/OnlineServicesCommon/Public/Online/OnlineServicesCommon.h
--- a/Engine/Plugins/Online/Experimental/OnlineServices/Source/OnlineServicesCommon/Public/Online/OnlineServicesCommon.h
+++ b/Engine/Plugins/Online/Experimental/OnlineServices/Source/OnlineServicesCommon/Public/Online/OnlineServicesCommon.h
@@ -40,6 +40,7 @@
 	virtual IPresencePtr GetPresenceInterface() override;
 	virtual IExternalUIPtr GetExternalUIInterface() override;
 	virtual ILobbiesPtr GetLobbiesInterface() override;
+	virtual ISessionsPtr GetSessionsInterface() override;
 	virtual IConnectivityPtr GetConnectivityInterface() override;
 	virtual IPrivilegesPtr GetPrivilegesInterface() override;
 	virtual TOnlineResult<FGetResolvedConnectString> GetResolvedConnectString(FGetResolvedConnectString::Params&& Params) override;
diff --git a/Engine/Plugins/Online/Experimental/OnlineServices/Source/OnlineServicesCommon/Public/Online/SessionsCommon.h b/Engine/Plugins/Online/Experimental/OnlineServices/Source/OnlineServicesCommon/Public/Online/SessionsCommon.h
new file mode 100644
--- /dev/null
+++ b/Engine/Plugins/Online/Experimental/OnlineServices/Source/OnlineServicesCommon/Public/Online/SessionsCommon.h
@@ -0,0 +1,59 @@
+// Copyright Epic Games, Inc. All Rights Reserved.
+
+#pragma once
+
+#include "Online/Sessions.h"
+#include "Online/OnlineComponent.h"
+
+namespace UE::Online {
+
+	class FOnlineServicesCommon;
+	class FAccountInfo;
+
+	struct FSessionEvents
+	{
+		TOnlineEventCallable<void(const FSessionJoined&)> OnSessionJoined;
+		TOnlineEventCallable<void(const FSessionLeft&)> OnSessionLeft;
+		TOnlineEventCallable<void(const FSessionUpdated&)> OnSessionUpdated;
+		TOnlineEventCallable<void(const FInviteReceived&)> OnInviteReceived;
+		TOnlineEventCallable<void(const FInviteAccepted&)> OnInviteAccepted;
+	};
+
+	class ONLINESERVICESCOMMON_API FSessionsCommon : public TOnlineComponent<ISessions>
+	{
+	public:
+		using Super = ISessions;
+
+		FSessionsCommon(FOnlineServicesCommon& InServices);
+
+		// TOnlineComponent
+		virtual void Initialize() override;
+		virtual void RegisterCommands() override;
+
+		// ISessions
+		virtual TOnlineResult<FGetAllSessions> GetAllSessions(FGetAllSessions::Params&& Params) override;
+		virtual TOnlineResult<FGetSessionByName> GetSessionByName(FGetSessionByName::Params&& Params) override;
+		virtual TOnlineResult<FGetSessionById> GetSessionById(FGetSessionById::Params&& Params) override;
+		virtual TOnlineAsyncOpHandle<FCreateSession> CreateSession(FCreateSession::Params&& Params) override;
+		virtual TOnlineAsyncOpHandle<FUpdateSession> UpdateSession(FUpdateSession::Params&& Params) override;
+		virtual TOnlineAsyncOpHandle<FLeaveSession> LeaveSession(FLeaveSession::Params&& Params) override;
+		virtual TOnlineAsyncOpHandle<FFindSessions> FindSessions(FFindSessions::Params&& Params) override;
+		virtual TOnlineAsyncOpHandle<FStartMatchmaking> StartMatchmaking(FStartMatchmaking::Params&& Params) override;
+		virtual TOnlineAsyncOpHandle<FJoinSession> JoinSession(FJoinSession::Params&& Params) override;
+		virtual TOnlineAsyncOpHandle<FSendSessionInvite> SendSessionInvite(FSendSessionInvite::Params&& Params) override;
+		virtual TOnlineAsyncOpHandle<FQuerySessionInvites> QuerySessionInvites(FQuerySessionInvites::Params&& Params) override;
+		virtual TOnlineAsyncOpHandle<FRejectSessionInvite> RejectSessionInvite(FRejectSessionInvite::Params&& Params) override;
+		virtual TOnlineAsyncOpHandle<FRegisterPlayers> RegisterPlayers(FRegisterPlayers::Params&& Params) override;
+		virtual TOnlineAsyncOpHandle<FUnregisterPlayers> UnregisterPlayers(FUnregisterPlayers::Params&& Params) override;
+
+		virtual TOnlineEvent<void(const FSessionJoined&)> OnSessionJoined() override;
+		virtual TOnlineEvent<void(const FSessionLeft&)> OnSessionLeft() override;
+		virtual TOnlineEvent<void(const FSessionUpdated&)> OnSessionUpdated() override;
+		virtual TOnlineEvent<void(const FInviteReceived&)> OnInviteReceived() override;
+		virtual TOnlineEvent<void(const FInviteAccepted&)> OnInviteAccepted() override;
+
+	protected:
+		FSessionEvents SessionEvents;
+};
+
+/* UE::Online */ }
diff --git a/Engine/Plugins/Online/Experimental/OnlineServices/Source/OnlineServicesInterface/Private/Online/Sessions.cpp b/Engine/Plugins/Online/Experimental/OnlineServices/Source/OnlineServicesInterface/Private/Online/Sessions.cpp
new file mode 100644
--- /dev/null
+++ b/Engine/Plugins/Online/Experimental/OnlineServices/Source/OnlineServicesInterface/Private/Online/Sessions.cpp
@@ -0,0 +1,256 @@
+// Copyright Epic Games, Inc. All Rights Reserved.
+
+#include "Online/Sessions.h"
+
+namespace UE::Online {
+
+const TCHAR* LexToString(ECustomSessionSettingVisibility Value)
+{
+	switch (Value)
+	{
+	case ECustomSessionSettingVisibility::ViaPingOnly:				return TEXT("ViaPingOnly");
+	case ECustomSessionSettingVisibility::ViaOnlineService:			return TEXT("ViaOnlineService");
+	case ECustomSessionSettingVisibility::ViaOnlineServiceAndPing:	return TEXT("ViaOnlineServiceAndPing");
+	default:														checkNoEntry(); // Intentional fallthrough
+	case ECustomSessionSettingVisibility::DontAdvertise:			return TEXT("DontAdvertise");
+	}
+}
+
+void LexFromString(ECustomSessionSettingVisibility& Value, const TCHAR* InStr)
+{
+	if (FCString::Stricmp(InStr, TEXT("DontAdvertise")) == 0)
+	{
+		Value = ECustomSessionSettingVisibility::DontAdvertise;
+	}
+	else if (FCString::Stricmp(InStr, TEXT("ViaPingOnly")) == 0)
+	{
+		Value = ECustomSessionSettingVisibility::ViaPingOnly;
+	}
+	else if (FCString::Stricmp(InStr, TEXT("ViaOnlineService")) == 0)
+	{
+		Value = ECustomSessionSettingVisibility::ViaOnlineService;
+	}
+	else if (FCString::Stricmp(InStr, TEXT("ViaOnlineServiceAndPing")) == 0)
+	{
+		Value = ECustomSessionSettingVisibility::ViaOnlineServiceAndPing;
+	}
+	else
+	{
+		checkNoEntry();
+		Value = ECustomSessionSettingVisibility::DontAdvertise;
+	}
+}
+
+const TCHAR* LexToString(ESessionState Value)
+{
+	switch (Value)
+	{
+	case ESessionState::Creating:	return TEXT("Creating");
+	case ESessionState::Valid:		return TEXT("Valid");
+	case ESessionState::Leaving:	return TEXT("Leaving");
+	case ESessionState::Destroying:	return TEXT("Destroying");
+	default:						checkNoEntry(); // Intentional fallthrough
+	case ESessionState::Invalid:	return TEXT("Invalid");
+	}
+}
+
+void LexFromString(ESessionState& Value, const TCHAR* InStr)
+{
+	if (FCString::Stricmp(InStr, TEXT("Invalid")) == 0)
+	{
+		Value = ESessionState::Invalid;
+	}
+	else if (FCString::Stricmp(InStr, TEXT("Creating")) == 0)
+	{
+		Value = ESessionState::Creating;
+	}
+	else if (FCString::Stricmp(InStr, TEXT("Valid")) == 0)
+	{
+		Value = ESessionState::Valid;
+	}
+	else if (FCString::Stricmp(InStr, TEXT("Leaving")) == 0)
+	{
+		Value = ESessionState::Leaving;
+	}
+	else if (FCString::Stricmp(InStr, TEXT("Destroying")) == 0)
+	{
+		Value = ESessionState::Destroying;
+	}
+	else
+	{
+		checkNoEntry();
+		Value = ESessionState::Invalid;
+	}
+}
+
+const TCHAR* LexToString(ESessionJoinPolicy Value)
+{
+	switch (Value)
+	{
+	case ESessionJoinPolicy::Public:		return TEXT("Public");
+	case ESessionJoinPolicy::FriendsOnly:	return TEXT("FriendsOnly");
+	default:								checkNoEntry(); // Intentional fallthrough
+	case ESessionJoinPolicy::InviteOnly:	return TEXT("InviteOnly");
+	}
+}
+
+void LexFromString(ESessionJoinPolicy& Value, const TCHAR* InStr)
+{
+	if (FCString::Stricmp(InStr, TEXT("Public")) == 0)
+	{
+		Value = ESessionJoinPolicy::Public;
+	}
+	else if (FCString::Stricmp(InStr, TEXT("FriendsOnly")) == 0)
+	{
+		Value = ESessionJoinPolicy::FriendsOnly;
+	}
+	else if (FCString::Stricmp(InStr, TEXT("InviteOnly")) == 0)
+	{
+		Value = ESessionJoinPolicy::InviteOnly;
+	}
+	else
+	{
+		checkNoEntry();
+		Value = ESessionJoinPolicy::InviteOnly;
+	}
+}
+
+FSessionMemberUpdate& FSessionMemberUpdate::operator+=(FSessionMemberUpdate&& UpdatedValue)
+{
+	for (TPair<FName, FCustomSessionSetting>& UpdatedMemberSetting : UpdatedValue.UpdatedMemberSettings)
+	{
+		// If an update adds a modification to a setting that had previously been marked for removal, we'll keep the latest change
+		RemovedMemberSettings.Remove(UpdatedMemberSetting.Key);
+	}
+	UpdatedMemberSettings.Append(MoveTemp(UpdatedValue.UpdatedMemberSettings));
+
+	for (FName& Key : UpdatedValue.RemovedMemberSettings)
+	{
+		// If an update removes a setting that had previously been modified, we'll keep the latest change
+		UpdatedMemberSettings.Remove(Key);
+		RemovedMemberSettings.AddUnique(MoveTemp(Key));
+	}
+
+	return *this;
+}
+
+FSessionSettingsUpdate& FSessionSettingsUpdate::operator+=(FSessionSettingsUpdate&& UpdatedValue)
+{
+	if (UpdatedValue.SchemaName.IsSet())
+	{
+		SchemaName = MoveTemp(UpdatedValue.SchemaName);
+	}
+
+	if (UpdatedValue.NumMaxPublicConnections.IsSet())
+	{
+		NumMaxPublicConnections = MoveTemp(UpdatedValue.NumMaxPublicConnections);
+	}
+
+	if (UpdatedValue.NumOpenPublicConnections.IsSet())
+	{
+		NumOpenPublicConnections = MoveTemp(UpdatedValue.NumOpenPublicConnections);
+	}
+
+	if (UpdatedValue.NumMaxPrivateConnections.IsSet())
+	{
+		NumMaxPrivateConnections = MoveTemp(UpdatedValue.NumMaxPrivateConnections);
+	}
+
+	if (UpdatedValue.NumOpenPrivateConnections.IsSet())
+	{
+		NumOpenPrivateConnections = MoveTemp(UpdatedValue.NumOpenPrivateConnections);
+	}
+
+	if (UpdatedValue.JoinPolicy.IsSet())
+	{
+		JoinPolicy = MoveTemp(UpdatedValue.JoinPolicy);
+	}
+
+	if (UpdatedValue.SessionIdOverride.IsSet())
+	{
+		SessionIdOverride = MoveTemp(UpdatedValue.SessionIdOverride);
+	}
+
+	if (UpdatedValue.IsLANSession.IsSet())
+	{
+		IsLANSession = MoveTemp(UpdatedValue.IsLANSession);
+	}
+
+	if (UpdatedValue.IsDedicatedServerSession.IsSet())
+	{
+		IsDedicatedServerSession = MoveTemp(UpdatedValue.IsDedicatedServerSession);
+	}
+
+	if (UpdatedValue.bAllowNewMembers.IsSet())
+	{
+		bAllowNewMembers = MoveTemp(UpdatedValue.bAllowNewMembers);
+	}
+
+	if (UpdatedValue.bAllowSanctionedPlayers.IsSet())
+	{
+		bAllowSanctionedPlayers = MoveTemp(UpdatedValue.bAllowSanctionedPlayers);
+	}
+
+	if (UpdatedValue.bAllowUnregisteredPlayers.IsSet())
+	{
+		bAllowUnregisteredPlayers = MoveTemp(UpdatedValue.bAllowUnregisteredPlayers);
+	}
+
+	if (UpdatedValue.bAntiCheatProtected.IsSet())
+	{
+		bAntiCheatProtected = MoveTemp(UpdatedValue.bAntiCheatProtected);
+	}
+
+	if (UpdatedValue.bPresenceEnabled.IsSet())
+	{
+		UpdatedValue.bPresenceEnabled = MoveTemp(UpdatedValue.bPresenceEnabled);
+	}
+
+	for (TPair<FName, FCustomSessionSetting>& UpdatedCustomSetting : UpdatedValue.UpdatedCustomSettings)
+	{
+		// If an update adds a modification to a setting that had previously been marked for removal, we'll keep the latest change
+		RemovedCustomSettings.Remove(UpdatedCustomSetting.Key);
+	}
+	UpdatedCustomSettings.Append(MoveTemp(UpdatedValue.UpdatedCustomSettings));
+
+	for (FName& Key : UpdatedValue.RemovedCustomSettings)
+	{
+		// If an update removes a setting that had previously been modified, we'll keep the latest change
+		UpdatedCustomSettings.Remove(Key);
+		RemovedCustomSettings.AddUnique(MoveTemp(Key));
+	}
+
+	for (TPair<FOnlineAccountIdHandle, FSessionMemberUpdate>& UpdatedSessionMember : UpdatedValue.UpdatedSessionMembers)
+	{
+		// If an update adds a modification to a member that had previously been marked for removal, we'll keep the latest change
+		RemovedSessionMembers.Remove(UpdatedSessionMember.Key);
+
+		FSessionMemberUpdate MemberUpdate = UpdatedSessionMembers.FindOrAdd(MoveTemp(UpdatedSessionMember.Key));
+		MemberUpdate += MoveTemp(UpdatedSessionMember.Value);
+	}
+
+	for (FOnlineAccountIdHandle& Key : UpdatedValue.RemovedSessionMembers)
+	{
+		// If an update removes a member that had previously been modified, we'll keep the latest change
+		UpdatedSessionMembers.Remove(Key);
+		RemovedSessionMembers.AddUnique(MoveTemp(Key));
+	}
+
+	for (TPair<FOnlineAccountIdHandle, FRegisteredUser>& UpdatedRegisteredUser : UpdatedValue.UpdatedRegisteredUsers)
+	{
+		// If an update adds a modification to a registered user that had previously been marked for removal, we'll keep the latest change
+		RemovedRegisteredUsers.Remove(UpdatedRegisteredUser.Key);
+	}
+	UpdatedRegisteredUsers.Append(MoveTemp(UpdatedValue.UpdatedRegisteredUsers));
+
+	for (FOnlineAccountIdHandle& Key : UpdatedValue.RemovedRegisteredUsers)
+	{
+		// If an update removes a registered user that had previously been modified, we'll keep the latest change
+		UpdatedRegisteredUsers.Remove(Key);
+		RemovedRegisteredUsers.AddUnique(MoveTemp(Key));
+	}
+
+	return *this;
+}
+
+/* UE::Online */ }
\ No newline at end of file
diff --git a/Engine/Plugins/Online/Experimental/OnlineServices/Source/OnlineServicesInterface/Public/Online/OnlineServices.h b/Engine/Plugins/Online/Experimental/OnlineServices/Source/OnlineServicesInterface/Public/Online/OnlineServices.h
--- a/Engine/Plugins/Online/Experimental/OnlineServices/Source/OnlineServicesInterface/Public/Online/OnlineServices.h
+++ b/Engine/Plugins/Online/Experimental/OnlineServices/Source/OnlineServicesInterface/Public/Online/OnlineServices.h
@@ -18,6 +18,7 @@
 using IPresencePtr = TSharedPtr<class IPresence>;
 using IExternalUIPtr = TSharedPtr<class IExternalUI>;
 using ILobbiesPtr = TSharedPtr<class ILobbies>;
+using ISessionsPtr = TSharedPtr<class ISessions>;
 using IConnectivityPtr = TSharedPtr<class IConnectivity>;
 using IPrivilegesPtr = TSharedPtr<class IPrivileges>;
 
@@ -29,6 +30,7 @@
 	{
 		FOnlineAccountIdHandle LocalUserId;
 		FOnlineLobbyIdHandle LobbyId;
+		FOnlineSessionIdHandle SessionId;
 		FName PortType;
 	};
 
@@ -78,6 +80,12 @@
 	virtual ILobbiesPtr GetLobbiesInterface() = 0;
 
 	/**
+	 * Get the sessions implementation
+	 * @return sessions implementation, may be null if not implemented for this service
+	 */
+	virtual ISessionsPtr GetSessionsInterface() = 0;
+
+	/**
 	 * 
 	 */
 	virtual IConnectivityPtr GetConnectivityInterface() = 0;
diff --git a/Engine/Plugins/Online/Experimental/OnlineServices/Source/OnlineServicesInterface/Public/Online/OnlineUtils.h b/Engine/Plugins/Online/Experimental/OnlineServices/Source/OnlineServicesInterface/Public/Online/OnlineUtils.h
--- a/Engine/Plugins/Online/Experimental/OnlineServices/Source/OnlineServicesInterface/Public/Online/OnlineUtils.h
+++ b/Engine/Plugins/Online/Experimental/OnlineServices/Source/OnlineServicesInterface/Public/Online/OnlineUtils.h
@@ -3,6 +3,7 @@
 #pragma once
 
 #include "Online/OnlineMeta.h"
+#include "Containers/UnrealString.h"
 
 namespace UE::Online {
 
@@ -22,6 +23,7 @@
 inline FString ToLogString(int32 Value);
 inline FString ToLogString(uint64 Value);
 inline FString ToLogString(int64 Value);
+inline FString ToLogString(bool Value);
 template <typename T> std::enable_if_t<!TModels<Meta::COnlineMetadataAvailable, T>::Value, FString> ToLogString(const T& Value);
 template <typename T> std::enable_if_t<TModels<Meta::COnlineMetadataAvailable, T>::Value, FString> ToLogString(const T& Value);
 
@@ -157,6 +159,11 @@
 	return FString::Printf(TEXT("%lli"), Value);
 }
 
+inline FString ToLogString(bool Value)
+{
+	return ::LexToString(Value);
+}
+
 inline FString ToLogString(FPlatformUserId PlatformUserId)
 {
 	return ToLogString(PlatformUserId.GetInternalId());
diff --git a/Engine/Plugins/Online/Experimental/OnlineServices/Source/OnlineServicesInterface/Public/Online/Sessions.h b/Engine/Plugins/Online/Experimental/OnlineServices/Source/OnlineServicesInterface/Public/Online/Sessions.h
new file mode 100644
--- /dev/null
+++ b/Engine/Plugins/Online/Experimental/OnlineServices/Source/OnlineServicesInterface/Public/Online/Sessions.h
@@ -0,0 +1,911 @@
+// Copyright Epic Games, Inc. All Rights Reserved.
+
+#pragma once
+
+#include "Online/OnlineAsyncOpHandle.h"
+#include "Online/CoreOnline.h"
+#include "Online/OnlineMeta.h"
+#include "Misc/TVariant.h"
+
+namespace UE::Online {
+
+struct FFindSessionsSearchFilter
+{
+	// TODO: Will share this type with Lobbies, so the implementation will be identical for now
+};
+
+enum class ECustomSessionSettingVisibility : uint8
+{
+	/* Don't advertise via the online service or QoS data */
+	DontAdvertise,
+	/* Advertise via the server ping data only */
+	ViaPingOnly,
+	/* Advertise via the online service only */
+	ViaOnlineService,
+	/* Advertise via the online service and via the ping data */
+	ViaOnlineServiceAndPing
+};
+ONLINESERVICESINTERFACE_API const TCHAR* LexToString(ECustomSessionSettingVisibility Value);
+ONLINESERVICESINTERFACE_API void LexFromString(ECustomSessionSettingVisibility& Value, const TCHAR* InStr);
+
+struct FCustomSessionSetting
+{
+	// TODO: Will share this type with Lobbies
+
+	using FVariantType = TVariant<FString, int64, double, bool>;
+
+	/** Setting value */
+	FVariantType Data;
+
+	/** How is this session setting advertised with the backend or searches */
+	ECustomSessionSettingVisibility Visibility;
+
+	/** Optional ID used in some platforms as the index instead of the setting name */
+	int32 ID;
+};
+
+using FCustomSessionSettingsMap = TMap<FName, FCustomSessionSetting>;
+
+struct FSessionMember
+{
+	FCustomSessionSettingsMap MemberSettings;
+};
+
+using FSessionMembersMap = TMap<FOnlineAccountIdHandle, FSessionMember>;
+
+struct FSessionMemberUpdate
+{
+	FCustomSessionSettingsMap UpdatedMemberSettings;
+	TArray<FName> RemovedMemberSettings;
+
+	FSessionMemberUpdate& operator+=(FSessionMemberUpdate&& UpdatedValue);
+};
+
+using FSessionMemberUpdatesMap = TMap<FOnlineAccountIdHandle, FSessionMemberUpdate>;
+
+enum class ESessionState : uint8
+{
+	Invalid,
+	Creating,
+	Valid,
+	Leaving,
+	Destroying
+};
+ONLINESERVICESINTERFACE_API const TCHAR* LexToString(ESessionState Value);
+ONLINESERVICESINTERFACE_API void LexFromString(ESessionState& Value, const TCHAR* InStr);
+
+enum class ESessionJoinPolicy : uint8
+{
+	Public,
+	FriendsOnly,
+	InviteOnly
+};
+ONLINESERVICESINTERFACE_API const TCHAR* LexToString(ESessionJoinPolicy Value);
+ONLINESERVICESINTERFACE_API void LexFromString(ESessionJoinPolicy& Value, const TCHAR* InStr);
+
+struct FRegisteredUser
+{
+	/* Whether a slot was reserved for this user upon registration */
+	bool bHasReservedSlot = false;
+
+	/* Whether the user is currently in the session */
+	bool bIsInSession = false;
+};
+
+using FRegisteredUsersMap = TMap<FOnlineAccountIdHandle, FRegisteredUser>;
+
+struct FSessionSettings
+{
+	/* The schema which will be applied to the session */
+	FName SchemaName;
+
+	/* Maximum number of public slots for session members */
+	int32 NumMaxPublicConnections;
+
+	/* Number of available public slots for session members */
+	int32 NumOpenPublicConnections;
+
+	/* Maximum number of private slots for session members */
+	int32 NumMaxPrivateConnections;
+
+	/* Number of available private slots for session members */
+	int32 NumOpenPrivateConnections;
+
+	/* Enum value describing the level of restriction to join the session */
+	ESessionJoinPolicy JoinPolicy;
+
+ 	/* In platforms that support this feature, it will set the session id to this value. Might be subject to minimum and maximum length */
+ 	FString SessionIdOverride;
+
+	/* Whether the session is only available in the local network and not via internet connection. Only available in some platforms */
+	bool IsLANSession;
+
+	/* Whether the session is configured to run as a dedicated server. Only available in some platforms */
+	bool IsDedicatedServerSession;
+
+	/* Whether new members are accepted in the session. Can vary depending on various factors, like the number of free slots available, or Join-In-Progress preferences when the session has started */
+	bool bAllowNewMembers;
+
+	/* Whether this session will allow sanctioned players to join it */
+	bool bAllowSanctionedPlayers;
+
+	/* Whether this session will allow unregistered players to join it */
+	bool bAllowUnregisteredPlayers;
+
+	/*Whether this is a secure session protected by anti-cheat services */
+	bool bAntiCheatProtected;
+
+	 /* Whether this session will show its information in presence updates. Can only be set in one session at a time */
+	 bool bPresenceEnabled;
+  
+ 	/* Map of user-defined settings to be passed to the platform APIs as additional information for various purposes */
+ 	FCustomSessionSettingsMap CustomSettings;
+ 
+ 	/* Map of session member ids to their corresponding user-defined settings */
+ 	FSessionMembersMap SessionMembers;
+ 
+ 	/* Map of registered users for this session. Can only be altered via RegisterPlayers and UnregisterPlayers */
+ 	FRegisteredUsersMap RegisteredUsers;
+};
+
+struct FSessionSettingsUpdate
+{
+	TOptional<FName> SchemaName;
+	TOptional<int32> NumMaxPublicConnections;
+	TOptional<int32> NumOpenPublicConnections;
+	TOptional<int32> NumMaxPrivateConnections;
+	TOptional<int32> NumOpenPrivateConnections;
+	TOptional<ESessionJoinPolicy> JoinPolicy;
+	TOptional<FString> SessionIdOverride;
+	TOptional<bool> IsLANSession;
+	TOptional<bool> IsDedicatedServerSession;
+	TOptional<bool> bAllowNewMembers;
+	TOptional<bool> bAllowSanctionedPlayers;
+	TOptional<bool> bAllowUnregisteredPlayers;
+	TOptional<bool> bAntiCheatProtected;
+	TOptional<bool> bPresenceEnabled;
+
+	FCustomSessionSettingsMap UpdatedCustomSettings;
+	TArray<FName> RemovedCustomSettings;
+
+	FSessionMemberUpdatesMap UpdatedSessionMembers;
+	TArray<FOnlineAccountIdHandle> RemovedSessionMembers;
+
+	FRegisteredUsersMap UpdatedRegisteredUsers;
+	TArray<FOnlineAccountIdHandle> RemovedRegisteredUsers;
+
+	FSessionSettingsUpdate& operator+=(FSessionSettingsUpdate&& UpdatedValue);
+};
+
+struct FSession
+{
+	/** The user who currently owns the session */
+	FOnlineAccountIdHandle OwnerUserId;
+
+	/** The id for the session, platform dependent */
+	FOnlineSessionIdHandle SessionId;
+
+	/** The current state of the session */
+	ESessionState CurrentState;
+
+	/** Set of session properties */
+	FSessionSettings SessionSettings;
+};
+
+struct FSessionInvite
+{
+	/* The user which sent the invite*/
+	FOnlineAccountIdHandle HostUserId;
+
+	/* The invite id, needed for rejecting the invite */
+	FName InviteId;
+
+	/* The session the invite is for */
+	TSharedRef<FSession> Session;
+};
+
+struct FGetAllSessions
+{
+	static constexpr TCHAR Name[] = TEXT("GetAllSessions");
+
+	struct Params
+	{
+
+	};
+
+	struct Result
+	{
+		TArray<TSharedRef<FSession>> Sessions;
+	};
+};
+
+struct FGetSessionByName
+{
+	static constexpr TCHAR Name[] = TEXT("GetSessionByName");
+
+	struct Params
+	{
+		FName LocalName;
+	};
+
+	struct Result
+	{
+		TSharedRef<FSession> Session;
+	};
+};
+
+struct FGetSessionById
+{
+	static constexpr TCHAR Name[] = TEXT("GetSessionById");
+
+	struct Params
+	{
+		FOnlineSessionIdHandle IdHandle;
+	};
+
+	struct Result
+	{
+		TSharedRef<FSession> Session;
+	};
+};
+
+struct FCreateSession
+{
+	static constexpr TCHAR Name[] = TEXT("CreateSession");
+
+	struct Params
+	{
+		/** The local user agent which will perform the action. */
+		FOnlineAccountIdHandle LocalUserId;
+
+		/** The local name for the session */
+		FName SessionName;
+
+		/** Settings object to define session properties during creation */
+		FSessionSettings SessionSettings;
+
+		/** Local users who will be joining the session on successful creation, along with the session creator */
+		FSessionMembersMap LocalUsers;
+	};
+
+	struct Result
+	{
+		TSharedRef<const FSession> Session;
+	};
+};
+
+struct FUpdateSession
+{
+	static constexpr TCHAR Name[] = TEXT("UpdateSession");
+
+	struct Params
+	{
+		/** The local user agent which will perform the action. */
+		FOnlineAccountIdHandle LocalUserId;
+
+		/** The local name for the session */
+		FName SessionName;
+
+		/** Changes to current session settings */
+		FSessionSettingsUpdate Mutations;
+	};
+
+	struct Result
+	{
+		TSharedRef<const FSession> Session;
+	};
+};
+
+struct FLeaveSession
+{
+	static constexpr TCHAR Name[] = TEXT("LeaveSession");
+
+	struct Params
+	{
+		/* The local user agent which leaves the session*/
+		FOnlineAccountIdHandle LocalUserId;
+
+		/* The local name for the session. */
+		FName SessionName;
+
+		/* Whether the call should attempt to destroy the session instead of just leave it */
+		bool bDestroySession;
+	};
+
+	struct Result
+	{
+		TSharedRef<const FSession> Session;
+	};
+};
+
+struct FFindSessions
+{
+	static constexpr TCHAR Name[] = TEXT("FindSessions");
+
+	struct Params
+	{
+		/* The local user agent which starts the session search*/
+		FOnlineAccountIdHandle LocalUserId;
+
+		/* Maximum number of results to return in one search */
+		uint32 MaxResults;
+
+		/* Filters to apply when searching for sessions. */
+		TArray<FFindSessionsSearchFilter> Filters;
+
+		/* Find sessions containing the target user. */
+		TOptional<FOnlineAccountIdHandle> TargetUser;
+
+		/* Find join info for the target session id. */
+		TOptional<FOnlineSessionIdHandle> SessionId;
+	};
+
+	struct Result
+	{
+		TArray<TSharedRef<const FSession>> FoundSessions;
+	};
+};
+
+struct FStartMatchmaking
+{
+	static constexpr TCHAR Name[] = TEXT("StartMatchmaking");
+
+	struct Params
+	{
+		/* Local users who will be joining the session */
+		FSessionMembersMap LocalUsers;
+
+		/* Local name for the session */
+		FName SessionName;
+
+		/* Preferred settings to be used during session creation */
+		FSessionSettings SessionSettings;
+
+		/* Filters to apply when searching for sessions */
+		TArray<FFindSessionsSearchFilter> SearchFilters;
+	};
+
+	struct Result
+	{
+		TSharedRef<const FSession> Session;
+	};
+};
+
+struct FJoinSession
+{
+	static constexpr TCHAR Name[] = TEXT("JoinSession");
+
+	struct Params
+	{
+		/* The local user agent which starts the join operation*/
+		FOnlineAccountIdHandle LocalUserId;
+
+		/* Local name for the session */
+		FName SessionName;
+
+		/* A reference to the session to be joined. */
+		TSharedRef<FSession> Session;
+
+		/* Local users who will be joining the session. */
+		FSessionMembersMap LocalUsers;
+	};
+
+	struct Result
+	{
+		TSharedRef<const FSession> Session;
+	};
+};
+
+struct FSendSessionInvite
+{
+	static constexpr TCHAR Name[] = TEXT("SendSessionInvite");
+
+	struct Params
+	{
+		/* The local user agent which sends the invite*/
+		FOnlineAccountIdHandle LocalUserId;
+
+		/* The local name for the session. */
+		FName SessionName;
+
+		/* Array of id handles for users to which the invites will be sent */
+		TArray<FOnlineAccountIdHandle> TargetUsers;
+	};
+
+	struct Result
+	{
+
+	};
+};
+
+struct FQuerySessionInvites
+{
+	static constexpr TCHAR Name[] = TEXT("QuerySessionInvites");
+
+	struct Params
+	{
+		/* The local user agent which started the query*/
+		FOnlineAccountIdHandle LocalUserId;
+	};
+
+	struct Result
+	{
+		/* Array of session invite objects built from the information returned in the query */
+		TArray<FSessionInvite> SessionInvites;
+	};
+};
+
+struct FRejectSessionInvite
+{
+	static constexpr TCHAR Name[] = TEXT("RejectSessionInvite");
+
+	struct Params
+	{
+		/* The local user agent which started the query*/
+		FOnlineAccountIdHandle LocalUserId;
+
+		/* The invite to be rejected */
+		FSessionInvite SessionInvite;
+	};
+
+	struct Result
+	{
+	};
+};
+
+struct FRegisterPlayers
+{
+	static constexpr TCHAR Name[] = TEXT("RegisterPlayers");
+
+	struct Params
+	{
+		/* The local name for the session. */
+		FName SessionName;
+
+		/* Array of users which will be registered */
+		TArray<FOnlineAccountIdHandle> TargetUsers;
+
+		/* Whether a slot should be saved for the registered users */
+		bool bReserveSlot;
+	};
+
+	struct Result
+	{
+	};
+};
+
+struct FUnregisterPlayers
+{
+	static constexpr TCHAR Name[] = TEXT("UnregisterPlayers");
+
+	struct Params
+	{
+		/* The local name for the session. */
+		FName SessionName;
+
+		/* Array of users which will be unregistered */
+		TArray<FOnlineAccountIdHandle> TargetUsers;
+
+		/* Whether unregistered players should be removed from the session, if they are in it */
+		bool bRemoveUnregisteredPlayers;
+	};
+
+	struct Result
+	{
+	};
+};
+
+/* Events */
+
+struct FSessionJoined
+{
+	/* The local users which joined the session */
+	TArray<FOnlineAccountIdHandle> LocalUserIds;
+
+	/* A shared reference to the session joined. */
+	TSharedRef<FSession> Session;
+};
+
+struct FSessionLeft
+{
+	/* The local users which left the session */
+	TArray<FOnlineAccountIdHandle> LocalUserIds;
+};
+
+struct FSessionUpdated
+{
+	/* Reference to the updated session object */
+	TSharedRef<const FSession> UpdatedSession;
+
+	/* Updated session settings */
+	FSessionSettingsUpdate UpdatedSettings;
+};
+
+struct FInviteReceived
+{
+	/* The local user which accepted the invite */
+	FOnlineAccountIdHandle LocalUserId;
+
+	/* Reference to the invite information */
+	TSharedRef<FSessionInvite> SessionInvite;
+};
+
+struct FInviteAccepted
+{
+	/* The local user which accepted the invite */
+	FOnlineAccountIdHandle LocalUserId;
+
+	/* Reference to the invite information */
+	TSharedRef<FSessionInvite> SessionInvite;
+};
+
+class ISessions
+{
+public:
+	/**
+	 * Get an array of all session objects.
+	 * 
+	 * @params Parameters for the GetAllSessions call
+	 * return
+	 */
+	virtual TOnlineResult<FGetAllSessions> GetAllSessions(FGetAllSessions::Params&& Params) = 0;
+
+	/**
+	 * Get the session object with a given local name.
+	 *
+	 * @params Parameters for the GetSessionByName call
+	 * return
+	 */
+	virtual TOnlineResult<FGetSessionByName> GetSessionByName(FGetSessionByName::Params&& Params) = 0;
+
+	/**
+	 * Get the session object with a given id handle.
+	 *
+	 * @params Parameters for the GetSessionById call
+	 * return
+	 */
+	virtual TOnlineResult<FGetSessionById> GetSessionById(FGetSessionById::Params&& Params) = 0;
+
+	/**
+	 * Create and join a new session.
+	 *
+	 * @param Parameters for the CreateSession call
+	 * @return
+	 */
+	virtual TOnlineAsyncOpHandle<FCreateSession> CreateSession(FCreateSession::Params&& Params) = 0;
+
+	/**
+	 * Update a given session's settings.
+	 *
+	 * @param Parameters for the UpdateSession call
+	 * @return
+	 */
+	virtual TOnlineAsyncOpHandle<FUpdateSession> UpdateSession(FUpdateSession::Params&& Params) = 0;
+
+	/**
+	 * Leave and optionally destroy a given session.
+	 *
+	 * @param Parameters for the LeaveSession call
+	 * @return
+	 */
+	virtual TOnlineAsyncOpHandle<FLeaveSession> LeaveSession(FLeaveSession::Params&& Params) = 0;
+
+	/**
+	 * Queries the API session service for sessions matching the given parameters.
+	 *
+	 * @param Parameters for the FindSessions call
+	 * @return
+	 */
+	virtual TOnlineAsyncOpHandle<FFindSessions> FindSessions(FFindSessions::Params&& Params) = 0;
+
+	/**
+	 * Starts the matchmaking process, which will either create a session with the passed parameters, or join one that matches the passed search filters.
+	 *
+	 * @param Parameters for the StartMatchmaking call
+	 * @return
+	 */
+	virtual TOnlineAsyncOpHandle<FStartMatchmaking> StartMatchmaking(FStartMatchmaking::Params&& Params) = 0;
+
+	/**
+	 * Starts the join process for the given session for all users provided.
+	 *
+	 * @param Parameters for the JoinSession call
+	 * @return
+	 */
+	virtual TOnlineAsyncOpHandle<FJoinSession> JoinSession(FJoinSession::Params&& Params) = 0;
+
+	/**
+	 * Sends an invite to the named session to all given users.
+	 *
+	 * @param Parameters for the SendSessionInvite call
+	 * @return
+	 */
+	virtual TOnlineAsyncOpHandle<FSendSessionInvite> SendSessionInvite(FSendSessionInvite::Params&& Params) = 0;
+
+	/**
+	 * Queries all session invites for a given user.
+	 *
+	 * @param Parameters for the QuerySessionInvites call
+	 * @return
+	 */
+	virtual TOnlineAsyncOpHandle<FQuerySessionInvites> QuerySessionInvites(FQuerySessionInvites::Params&& Params) = 0;
+
+	/**
+	 * Rejects a given session invite for a user.
+	 *
+	 * @param Parameters for the RejectSessionInvite call
+	 * @return
+	 */
+	virtual TOnlineAsyncOpHandle<FRejectSessionInvite> RejectSessionInvite(FRejectSessionInvite::Params&& Params) = 0;
+
+	/**
+	 * Registers given players in the named session.
+	 *
+	 * @param Parameters for the RegisterPlayers call
+	 * @return
+	 */
+	virtual TOnlineAsyncOpHandle<FRegisterPlayers> RegisterPlayers(FRegisterPlayers::Params&& Params) = 0;
+
+	/**
+	 * Unregisters given players from the named session.
+	 *
+	 * @param Parameters for the UnregisterPlayers call
+	 * @return
+	 */
+	virtual TOnlineAsyncOpHandle<FUnregisterPlayers> UnregisterPlayers(FUnregisterPlayers::Params&& Params) = 0;
+
+	/* Events */
+
+	/**
+	 * Get the event that is triggered when a session is joined.
+	 * This event will trigger as a result of creating or joining a session.
+	 *
+	 * @return
+	 */
+	virtual TOnlineEvent<void(const FSessionJoined&)> OnSessionJoined() = 0;
+
+	/**
+	 * Get the event that is triggered when a session is left.
+	 * This event will trigger as a result of leaving or destroying a session.
+	 *
+	 * @return
+	 */
+	virtual TOnlineEvent<void(const FSessionLeft&)> OnSessionLeft() = 0;
+
+	/**
+	 * Get the event that is triggered when a session invite is accepted.
+	 * This event will trigger as a result of accepting a platform session invite.
+	 *
+	 * @return
+	 */
+	virtual TOnlineEvent<void(const FSessionUpdated&)> OnSessionUpdated() = 0;
+
+	/**
+	 * Get the event that is triggered when a session invite is received.
+	 * This event will trigger as a result of receiving a platform session invite.
+	 *
+	 * @return
+	 */
+	virtual TOnlineEvent<void(const FInviteReceived&)> OnInviteReceived() = 0;
+
+	/**
+	 * Get the event that is triggered when a session invite is accepted.
+	 * This event will trigger as a result of accepting a platform session invite.
+	 *
+	 * @return
+	 */
+	virtual TOnlineEvent<void(const FInviteAccepted&)> OnInviteAccepted() = 0;
+};
+
+namespace Meta {
+
+BEGIN_ONLINE_STRUCT_META(FFindSessionsSearchFilter)
+END_ONLINE_STRUCT_META()
+
+BEGIN_ONLINE_STRUCT_META(FCustomSessionSetting)
+END_ONLINE_STRUCT_META()
+
+BEGIN_ONLINE_STRUCT_META(FSessionMember)
+	ONLINE_STRUCT_FIELD(FSessionMember, MemberSettings)
+END_ONLINE_STRUCT_META()
+
+BEGIN_ONLINE_STRUCT_META(FSessionMemberUpdate)
+	ONLINE_STRUCT_FIELD(FSessionMemberUpdate, UpdatedMemberSettings),
+	ONLINE_STRUCT_FIELD(FSessionMemberUpdate, RemovedMemberSettings)
+END_ONLINE_STRUCT_META()
+
+BEGIN_ONLINE_STRUCT_META(FRegisteredUser)
+	ONLINE_STRUCT_FIELD(FRegisteredUser, bHasReservedSlot),
+	ONLINE_STRUCT_FIELD(FRegisteredUser, bIsInSession)
+END_ONLINE_STRUCT_META()
+
+BEGIN_ONLINE_STRUCT_META(FSessionSettings)
+	ONLINE_STRUCT_FIELD(FSessionSettings, SchemaName),
+	ONLINE_STRUCT_FIELD(FSessionSettings, NumMaxPublicConnections),
+	ONLINE_STRUCT_FIELD(FSessionSettings, NumOpenPublicConnections),
+	ONLINE_STRUCT_FIELD(FSessionSettings, NumMaxPrivateConnections),
+	ONLINE_STRUCT_FIELD(FSessionSettings, NumOpenPrivateConnections),
+	ONLINE_STRUCT_FIELD(FSessionSettings, JoinPolicy),
+	ONLINE_STRUCT_FIELD(FSessionSettings, SessionIdOverride),
+	ONLINE_STRUCT_FIELD(FSessionSettings, IsLANSession),
+	ONLINE_STRUCT_FIELD(FSessionSettings, IsDedicatedServerSession),
+	ONLINE_STRUCT_FIELD(FSessionSettings, bAllowNewMembers),
+	ONLINE_STRUCT_FIELD(FSessionSettings, bAllowSanctionedPlayers),
+	ONLINE_STRUCT_FIELD(FSessionSettings, bAllowUnregisteredPlayers),
+	ONLINE_STRUCT_FIELD(FSessionSettings, bAntiCheatProtected),
+	ONLINE_STRUCT_FIELD(FSessionSettings, bPresenceEnabled),
+	ONLINE_STRUCT_FIELD(FSessionSettings, CustomSettings),
+	ONLINE_STRUCT_FIELD(FSessionSettings, SessionMembers),
+	ONLINE_STRUCT_FIELD(FSessionSettings, RegisteredUsers)
+END_ONLINE_STRUCT_META()
+
+BEGIN_ONLINE_STRUCT_META(FSessionSettingsUpdate)
+	ONLINE_STRUCT_FIELD(FSessionSettingsUpdate, SchemaName),
+	ONLINE_STRUCT_FIELD(FSessionSettingsUpdate, NumMaxPublicConnections),
+	ONLINE_STRUCT_FIELD(FSessionSettingsUpdate, NumOpenPublicConnections),
+	ONLINE_STRUCT_FIELD(FSessionSettingsUpdate, NumMaxPrivateConnections),
+	ONLINE_STRUCT_FIELD(FSessionSettingsUpdate, NumOpenPrivateConnections),
+	ONLINE_STRUCT_FIELD(FSessionSettingsUpdate, JoinPolicy),
+	ONLINE_STRUCT_FIELD(FSessionSettingsUpdate, SessionIdOverride),
+	ONLINE_STRUCT_FIELD(FSessionSettingsUpdate, IsLANSession),
+	ONLINE_STRUCT_FIELD(FSessionSettingsUpdate, IsDedicatedServerSession),
+	ONLINE_STRUCT_FIELD(FSessionSettingsUpdate, bAllowNewMembers),
+	ONLINE_STRUCT_FIELD(FSessionSettingsUpdate, bAllowSanctionedPlayers),
+	ONLINE_STRUCT_FIELD(FSessionSettingsUpdate, bAllowUnregisteredPlayers),
+	ONLINE_STRUCT_FIELD(FSessionSettingsUpdate, bAntiCheatProtected),
+	ONLINE_STRUCT_FIELD(FSessionSettingsUpdate, bPresenceEnabled),
+	ONLINE_STRUCT_FIELD(FSessionSettingsUpdate, UpdatedCustomSettings),
+	ONLINE_STRUCT_FIELD(FSessionSettingsUpdate, RemovedCustomSettings),
+	ONLINE_STRUCT_FIELD(FSessionSettingsUpdate, UpdatedSessionMembers),
+	ONLINE_STRUCT_FIELD(FSessionSettingsUpdate, RemovedSessionMembers),
+	ONLINE_STRUCT_FIELD(FSessionSettingsUpdate, UpdatedRegisteredUsers),
+	ONLINE_STRUCT_FIELD(FSessionSettingsUpdate, RemovedRegisteredUsers)
+END_ONLINE_STRUCT_META()
+
+BEGIN_ONLINE_STRUCT_META(FSession)
+	ONLINE_STRUCT_FIELD(FSession, OwnerUserId),
+	ONLINE_STRUCT_FIELD(FSession, SessionId),
+	ONLINE_STRUCT_FIELD(FSession, CurrentState),
+	ONLINE_STRUCT_FIELD(FSession, SessionSettings)
+END_ONLINE_STRUCT_META()
+
+BEGIN_ONLINE_STRUCT_META(FSessionInvite)
+	ONLINE_STRUCT_FIELD(FSessionInvite, HostUserId),
+	ONLINE_STRUCT_FIELD(FSessionInvite, InviteId),
+	ONLINE_STRUCT_FIELD(FSessionInvite, Session)
+END_ONLINE_STRUCT_META()
+
+BEGIN_ONLINE_STRUCT_META(FGetAllSessions::Params)
+END_ONLINE_STRUCT_META()
+
+BEGIN_ONLINE_STRUCT_META(FGetAllSessions::Result)
+	ONLINE_STRUCT_FIELD(FGetAllSessions::Result, Sessions)
+END_ONLINE_STRUCT_META()
+
+BEGIN_ONLINE_STRUCT_META(FGetSessionByName::Params)
+	ONLINE_STRUCT_FIELD(FGetSessionByName::Params, LocalName)
+	END_ONLINE_STRUCT_META()
+
+BEGIN_ONLINE_STRUCT_META(FGetSessionByName::Result)
+	ONLINE_STRUCT_FIELD(FGetSessionByName::Result, Session)
+END_ONLINE_STRUCT_META()
+
+BEGIN_ONLINE_STRUCT_META(FGetSessionById::Params)
+	ONLINE_STRUCT_FIELD(FGetSessionById::Params, IdHandle)
+END_ONLINE_STRUCT_META()
+
+BEGIN_ONLINE_STRUCT_META(FGetSessionById::Result)
+	ONLINE_STRUCT_FIELD(FGetSessionById::Result, Session)
+END_ONLINE_STRUCT_META()
+
+BEGIN_ONLINE_STRUCT_META(FCreateSession::Params)
+	ONLINE_STRUCT_FIELD(FCreateSession::Params, LocalUserId),
+	ONLINE_STRUCT_FIELD(FCreateSession::Params, SessionName),
+	ONLINE_STRUCT_FIELD(FCreateSession::Params, SessionSettings),
+	ONLINE_STRUCT_FIELD(FCreateSession::Params, LocalUsers)
+END_ONLINE_STRUCT_META()
+
+BEGIN_ONLINE_STRUCT_META(FCreateSession::Result)
+	ONLINE_STRUCT_FIELD(FCreateSession::Result, Session)
+END_ONLINE_STRUCT_META()
+
+BEGIN_ONLINE_STRUCT_META(FUpdateSession::Params)
+	ONLINE_STRUCT_FIELD(FUpdateSession::Params, LocalUserId),
+	ONLINE_STRUCT_FIELD(FUpdateSession::Params, SessionName),
+	ONLINE_STRUCT_FIELD(FUpdateSession::Params, Mutations)
+END_ONLINE_STRUCT_META()
+
+BEGIN_ONLINE_STRUCT_META(FUpdateSession::Result)
+	ONLINE_STRUCT_FIELD(FUpdateSession::Result, Session)
+END_ONLINE_STRUCT_META()
+
+BEGIN_ONLINE_STRUCT_META(FLeaveSession::Params)
+	ONLINE_STRUCT_FIELD(FLeaveSession::Params, LocalUserId),
+	ONLINE_STRUCT_FIELD(FLeaveSession::Params, SessionName),
+	ONLINE_STRUCT_FIELD(FLeaveSession::Params, bDestroySession)
+END_ONLINE_STRUCT_META()
+
+BEGIN_ONLINE_STRUCT_META(FLeaveSession::Result)
+	ONLINE_STRUCT_FIELD(FLeaveSession::Result, Session)
+END_ONLINE_STRUCT_META()
+
+BEGIN_ONLINE_STRUCT_META(FFindSessions::Params)
+	ONLINE_STRUCT_FIELD(FFindSessions::Params, LocalUserId),
+	ONLINE_STRUCT_FIELD(FFindSessions::Params, MaxResults),
+	ONLINE_STRUCT_FIELD(FFindSessions::Params, Filters),
+	ONLINE_STRUCT_FIELD(FFindSessions::Params, TargetUser),
+	ONLINE_STRUCT_FIELD(FFindSessions::Params, SessionId)
+END_ONLINE_STRUCT_META()
+
+BEGIN_ONLINE_STRUCT_META(FFindSessions::Result)
+	ONLINE_STRUCT_FIELD(FFindSessions::Result, FoundSessions)
+END_ONLINE_STRUCT_META()
+
+BEGIN_ONLINE_STRUCT_META(FStartMatchmaking::Params)
+	ONLINE_STRUCT_FIELD(FStartMatchmaking::Params, LocalUsers),
+	ONLINE_STRUCT_FIELD(FStartMatchmaking::Params, SessionName),
+	ONLINE_STRUCT_FIELD(FStartMatchmaking::Params, SessionSettings),
+	ONLINE_STRUCT_FIELD(FStartMatchmaking::Params, SearchFilters)
+END_ONLINE_STRUCT_META()
+
+BEGIN_ONLINE_STRUCT_META(FStartMatchmaking::Result)
+	ONLINE_STRUCT_FIELD(FStartMatchmaking::Result, Session)
+END_ONLINE_STRUCT_META()
+
+BEGIN_ONLINE_STRUCT_META(FJoinSession::Params)
+	ONLINE_STRUCT_FIELD(FJoinSession::Params, LocalUserId),
+	ONLINE_STRUCT_FIELD(FJoinSession::Params, SessionName),
+	ONLINE_STRUCT_FIELD(FJoinSession::Params, Session),
+	ONLINE_STRUCT_FIELD(FJoinSession::Params, LocalUsers)
+END_ONLINE_STRUCT_META()
+
+BEGIN_ONLINE_STRUCT_META(FJoinSession::Result)
+	ONLINE_STRUCT_FIELD(FJoinSession::Result, Session)
+END_ONLINE_STRUCT_META()
+
+BEGIN_ONLINE_STRUCT_META(FSendSessionInvite::Params)
+	ONLINE_STRUCT_FIELD(FSendSessionInvite::Params, LocalUserId),
+	ONLINE_STRUCT_FIELD(FSendSessionInvite::Params, SessionName),
+	ONLINE_STRUCT_FIELD(FSendSessionInvite::Params, TargetUsers)
+END_ONLINE_STRUCT_META()
+
+BEGIN_ONLINE_STRUCT_META(FSendSessionInvite::Result)
+END_ONLINE_STRUCT_META()
+
+BEGIN_ONLINE_STRUCT_META(FQuerySessionInvites::Params)
+	ONLINE_STRUCT_FIELD(FQuerySessionInvites::Params, LocalUserId)
+END_ONLINE_STRUCT_META()
+
+BEGIN_ONLINE_STRUCT_META(FQuerySessionInvites::Result)
+	ONLINE_STRUCT_FIELD(FQuerySessionInvites::Result, SessionInvites)
+END_ONLINE_STRUCT_META()
+
+BEGIN_ONLINE_STRUCT_META(FRejectSessionInvite::Params)
+	ONLINE_STRUCT_FIELD(FRejectSessionInvite::Params, LocalUserId),
+	ONLINE_STRUCT_FIELD(FRejectSessionInvite::Params, SessionInvite)
+END_ONLINE_STRUCT_META()
+
+BEGIN_ONLINE_STRUCT_META(FRejectSessionInvite::Result)
+END_ONLINE_STRUCT_META()
+
+BEGIN_ONLINE_STRUCT_META(FRegisterPlayers::Params)
+	ONLINE_STRUCT_FIELD(FRegisterPlayers::Params, SessionName),
+	ONLINE_STRUCT_FIELD(FRegisterPlayers::Params, TargetUsers),
+	ONLINE_STRUCT_FIELD(FRegisterPlayers::Params, bReserveSlot)
+END_ONLINE_STRUCT_META()
+
+BEGIN_ONLINE_STRUCT_META(FRegisterPlayers::Result)
+END_ONLINE_STRUCT_META()
+
+BEGIN_ONLINE_STRUCT_META(FUnregisterPlayers::Params)
+	ONLINE_STRUCT_FIELD(FUnregisterPlayers::Params, SessionName),
+	ONLINE_STRUCT_FIELD(FUnregisterPlayers::Params, TargetUsers),
+	ONLINE_STRUCT_FIELD(FUnregisterPlayers::Params, bRemoveUnregisteredPlayers)
+END_ONLINE_STRUCT_META()
+
+BEGIN_ONLINE_STRUCT_META(FUnregisterPlayers::Result)
+END_ONLINE_STRUCT_META()
+
+/* Meta*/ }
+
+/* UE::Online */ }
\ No newline at end of file
diff --git a/Engine/Source/Runtime/CoreOnline/Private/Online/CoreOnline.cpp b/Engine/Source/Runtime/CoreOnline/Private/Online/CoreOnline.cpp
--- a/Engine/Source/Runtime/CoreOnline/Private/Online/CoreOnline.cpp
+++ b/Engine/Source/Runtime/CoreOnline/Private/Online/CoreOnline.cpp
@@ -300,6 +300,11 @@
 	return FString::Printf(TEXT("%s:%d"), LexToString(Id.GetOnlineServicesType()), Id.GetHandle());
 }
 
+FString ToLogString(const FOnlineSessionIdHandle& Id)
+{
+	return FString::Printf(TEXT("%s:%d"), LexToString(Id.GetOnlineServicesType()), Id.GetHandle());
+}
+
 }	/* UE::Online */
 
 FString FUniqueNetIdWrapper::ToDebugString() const
diff --git a/Engine/Source/Runtime/CoreOnline/Public/Online/CoreOnline.h b/Engine/Source/Runtime/CoreOnline/Public/Online/CoreOnline.h
--- a/Engine/Source/Runtime/CoreOnline/Public/Online/CoreOnline.h
+++ b/Engine/Source/Runtime/CoreOnline/Public/Online/CoreOnline.h
@@ -291,9 +291,11 @@
 
 using FOnlineAccountIdHandle = TOnlineIdHandle<OnlineIdHandleTags::FAccount>;
 using FOnlineLobbyIdHandle = TOnlineIdHandle<OnlineIdHandleTags::FLobby>;
+using FOnlineSessionIdHandle = TOnlineIdHandle<OnlineIdHandleTags::FSession>;
 
 COREONLINE_API FString ToLogString(const FOnlineAccountIdHandle& Id);
 COREONLINE_API FString ToLogString(const FOnlineLobbyIdHandle& Id);
+COREONLINE_API FString ToLogString(const FOnlineSessionIdHandle& Id);
 
 template<typename IdType>
 inline uint32 GetTypeHash(const TOnlineIdHandle<IdType>& Handle)
@@ -312,6 +314,7 @@
 };
 
 using IOnlineAccountIdRegistry = IOnlineIdRegistry<OnlineIdHandleTags::FAccount>;
+using IOnlineSessionIdRegistry = IOnlineIdRegistry<OnlineIdHandleTags::FSession>;
 
 class FOnlineIdRegistryRegistry
 {
