- Achievements interface and common impl

diff --git a/Engine/Plugins/Online/Experimental/OnlineServices/Source/OnlineServicesCommon/Private/Online/AchievementsCommon.cpp b/Engine/Plugins/Online/Experimental/OnlineServices/Source/OnlineServicesCommon/Private/Online/AchievementsCommon.cpp
new file mode 100644
--- /dev/null
+++ b/Engine/Plugins/Online/Experimental/OnlineServices/Source/OnlineServicesCommon/Private/Online/AchievementsCommon.cpp
@@ -0,0 +1,73 @@
+// Copyright Epic Games, Inc. All Rights Reserved.
+
+#include "Online/AchievementsCommon.h"
+
+#include "Online/OnlineAsyncOp.h"
+#include "Online/OnlineErrorDefinitions.h"
+#include "Online/OnlineServicesCommon.h"
+
+namespace UE::Online {
+
+FAchievementsCommon::FAchievementsCommon(FOnlineServicesCommon& InServices)
+	: TOnlineComponent(TEXT("Achievements"), InServices)
+{
+}
+
+void FAchievementsCommon::RegisterCommands()
+{
+	RegisterCommand(&FAchievementsCommon::QueryAchievementDefinitions);
+	RegisterCommand(&FAchievementsCommon::GetAchievementIds);
+	RegisterCommand(&FAchievementsCommon::GetAchievementDefinition);
+	RegisterCommand(&FAchievementsCommon::QueryAchievementStates);
+	RegisterCommand(&FAchievementsCommon::GetAchievementState);
+	RegisterCommand(&FAchievementsCommon::UnlockAchievements);
+	RegisterCommand(&FAchievementsCommon::DisplayAchievementUI);
+}
+
+TOnlineAsyncOpHandle<FQueryAchievementDefinitions> FAchievementsCommon::QueryAchievementDefinitions(FQueryAchievementDefinitions::Params&& Params)
+{
+	TOnlineAsyncOpRef<FQueryAchievementDefinitions> Operation = GetOp<FQueryAchievementDefinitions>(MoveTemp(Params));
+	Operation->SetError(Errors::NotImplemented());  
+	return Operation->GetHandle();
+}
+
+TOnlineResult<FGetAchievementIds> FAchievementsCommon::GetAchievementIds(FGetAchievementIds::Params&& Params)
+{
+	return TOnlineResult<FGetAchievementIds>(Errors::NotImplemented());
+}
+
+TOnlineResult<FGetAchievementDefinition> FAchievementsCommon::GetAchievementDefinition(FGetAchievementDefinition::Params&& Params)
+{
+	return TOnlineResult<FGetAchievementDefinition>(Errors::NotImplemented());
+}
+
+TOnlineAsyncOpHandle<FQueryAchievementStates> FAchievementsCommon::QueryAchievementStates(FQueryAchievementStates::Params&& Params)
+{
+	TOnlineAsyncOpRef<FQueryAchievementStates> Operation = GetOp<FQueryAchievementStates>(MoveTemp(Params));
+	Operation->SetError(Errors::NotImplemented());
+	return Operation->GetHandle();
+}
+
+TOnlineResult<FGetAchievementState> FAchievementsCommon::GetAchievementState(FGetAchievementState::Params&& Params)
+{
+	return TOnlineResult<FGetAchievementState>(Errors::NotImplemented());
+}
+
+TOnlineAsyncOpHandle<FUnlockAchievements> FAchievementsCommon::UnlockAchievements(FUnlockAchievements::Params&& Params)
+{
+	TOnlineAsyncOpRef<FUnlockAchievements> Operation = GetOp<FUnlockAchievements>(MoveTemp(Params));
+	Operation->SetError(Errors::NotImplemented());
+	return Operation->GetHandle();
+}
+
+TOnlineResult<FDisplayAchievementUI> FAchievementsCommon::DisplayAchievementUI(FDisplayAchievementUI::Params&& Params)
+{
+	return TOnlineResult<FDisplayAchievementUI>(Errors::NotImplemented());
+}
+
+TOnlineEvent<void(const FAchievementStateUpdated&)> FAchievementsCommon::OnAchievementStateUpdated()
+{
+	return OnAchievementStateUpdatedEvent;
+}
+
+/* UE::Online */ }
diff --git a/Engine/Plugins/Online/Experimental/OnlineServices/Source/OnlineServicesCommon/Private/Online/OnlineServicesCommon.cpp b/Engine/Plugins/Online/Experimental/OnlineServices/Source/OnlineServicesCommon/Private/Online/OnlineServicesCommon.cpp
--- a/Engine/Plugins/Online/Experimental/OnlineServices/Source/OnlineServicesCommon/Private/Online/OnlineServicesCommon.cpp
+++ b/Engine/Plugins/Online/Experimental/OnlineServices/Source/OnlineServicesCommon/Private/Online/OnlineServicesCommon.cpp
@@ -40,6 +40,11 @@
 	Shutdown();
 }
 
+IAchievementsPtr FOnlineServicesCommon::GetAchievementsInterface()
+{
+	return IAchievementsPtr(AsShared(), Get<IAchievements>());
+}
+
 IAuthPtr FOnlineServicesCommon::GetAuthInterface()
 {
 	return IAuthPtr(AsShared(), Get<IAuth>());
diff --git a/Engine/Plugins/Online/Experimental/OnlineServices/Source/OnlineServicesCommon/Public/Online/AchievementsCommon.h b/Engine/Plugins/Online/Experimental/OnlineServices/Source/OnlineServicesCommon/Public/Online/AchievementsCommon.h
new file mode 100644
--- /dev/null
+++ b/Engine/Plugins/Online/Experimental/OnlineServices/Source/OnlineServicesCommon/Public/Online/AchievementsCommon.h
@@ -0,0 +1,36 @@
+// Copyright Epic Games, Inc. All Rights Reserved.
+
+#pragma once
+
+#include "Online/Achievements.h"
+#include "Online/OnlineComponent.h"
+
+namespace UE::Online {
+
+class FOnlineServicesCommon;
+
+class ONLINESERVICESCOMMON_API FAchievementsCommon : public TOnlineComponent<IAchievements>
+{
+public:
+	using Super = IAchievements;
+
+	FAchievementsCommon(FOnlineServicesCommon& InServices);
+
+	// TOnlineComponent
+	virtual void RegisterCommands() override;
+
+	// IAchievements
+	virtual TOnlineAsyncOpHandle<FQueryAchievementDefinitions> QueryAchievementDefinitions(FQueryAchievementDefinitions::Params&& Params) override;
+	virtual TOnlineResult<FGetAchievementIds> GetAchievementIds(FGetAchievementIds::Params&& Params) override;
+	virtual TOnlineResult<FGetAchievementDefinition> GetAchievementDefinition(FGetAchievementDefinition::Params&& Params) override;
+	virtual TOnlineAsyncOpHandle<FQueryAchievementStates> QueryAchievementStates(FQueryAchievementStates::Params&& Params) override;
+	virtual TOnlineResult<FGetAchievementState> GetAchievementState(FGetAchievementState::Params&& Params) override;
+	virtual TOnlineAsyncOpHandle<FUnlockAchievements> UnlockAchievements(FUnlockAchievements::Params&& Params) override;
+	virtual TOnlineResult<FDisplayAchievementUI> DisplayAchievementUI(FDisplayAchievementUI::Params&& Params) override;
+	virtual TOnlineEvent<void(const FAchievementStateUpdated&)> OnAchievementStateUpdated() override;
+
+protected:
+	TOnlineEventCallable<void(const FAchievementStateUpdated&)> OnAchievementStateUpdatedEvent;
+};
+
+/* UE::Online */ }
diff --git a/Engine/Plugins/Online/Experimental/OnlineServices/Source/OnlineServicesCommon/Public/Online/OnlineServicesCommon.h b/Engine/Plugins/Online/Experimental/OnlineServices/Source/OnlineServicesCommon/Public/Online/OnlineServicesCommon.h
--- a/Engine/Plugins/Online/Experimental/OnlineServices/Source/OnlineServicesCommon/Public/Online/OnlineServicesCommon.h
+++ b/Engine/Plugins/Online/Experimental/OnlineServices/Source/OnlineServicesCommon/Public/Online/OnlineServicesCommon.h
@@ -35,6 +35,7 @@
 	// IOnlineServices
 	virtual void Init() override;
 	virtual void Destroy() override;
+	virtual IAchievementsPtr GetAchievementsInterface() override;
 	virtual IAuthPtr GetAuthInterface() override;
 	virtual IFriendsPtr GetFriendsInterface() override;
 	virtual IPresencePtr GetPresenceInterface() override;
diff --git a/Engine/Plugins/Online/Experimental/OnlineServices/Source/OnlineServicesInterface/Private/Online/Achievements.cpp b/Engine/Plugins/Online/Experimental/OnlineServices/Source/OnlineServicesInterface/Private/Online/Achievements.cpp
new file mode 100644
--- /dev/null
+++ b/Engine/Plugins/Online/Experimental/OnlineServices/Source/OnlineServicesInterface/Private/Online/Achievements.cpp
@@ -0,0 +1,47 @@
+// Copyright Epic Games, Inc. All Rights Reserved.
+
+#include "Online/Achievements.h"
+
+namespace UE::Online {
+
+FString ToLogString(const FAchievementStatDefinition& Definition)
+{
+	return FString::Printf(TEXT("StatId=\"%s\" UnlockThreshold=%d"),
+		*Definition.StatId, Definition.UnlockThreshold);
+}
+
+FString ToLogString(const FAchievementDefinition& Definition)
+{
+	const TCHAR FormatString[] = TEXT(
+		"AchievementId=\"%s\" "
+		"UnlockedDisplayName=\"%s\" "
+		"UnlockedDescription=\"%s\" "
+		"LockedDisplayName=\"%s\" "
+		"LockedDescription=\"%s\" "
+		"FlavorText=\"%s\" "
+		"UnlockedIconUrl=\"%s\" "
+		"LockedIconUrl=\"%s\" "
+		"bIsHidden=%s "
+		"StatDefinitions=[%s]");
+
+	return FString::Printf(FormatString,
+		*Definition.AchievementId,
+		*Definition.UnlockedDisplayName.ToString(),
+		*Definition.UnlockedDescription.ToString(),
+		*Definition.LockedDisplayName.ToString(),
+		*Definition.LockedDescription.ToString(),
+		*Definition.FlavorText.ToString(),
+		*Definition.UnlockedIconUrl,
+		*Definition.LockedIconUrl,
+		*::LexToString(Definition.bIsHidden),
+		*FString::JoinBy(Definition.StatDefinitions, TEXT(", "), [](const FAchievementStatDefinition& Definition) { return ToLogString(Definition); })
+	);
+}
+
+FString ToLogString(const FAchievementState& State)
+{
+	return FString::Printf(TEXT("AchievementId=\"%s\" Progress=%.2f UnlockTime=%s"),
+		*State.AchievementId, State.Progress, *State.UnlockTime.ToString());
+}
+
+/* UE::Online */ }
diff --git a/Engine/Plugins/Online/Experimental/OnlineServices/Source/OnlineServicesInterface/Public/Online/Achievements.h b/Engine/Plugins/Online/Experimental/OnlineServices/Source/OnlineServicesInterface/Public/Online/Achievements.h
new file mode 100644
--- /dev/null
+++ b/Engine/Plugins/Online/Experimental/OnlineServices/Source/OnlineServicesInterface/Public/Online/Achievements.h
@@ -0,0 +1,325 @@
+// Copyright Epic Games, Inc. All Rights Reserved.
+
+#pragma once
+
+#include "Online/CoreOnline.h"
+#include "Online/OnlineAsyncOpHandle.h"
+#include "Online/OnlineMeta.h"
+#include "Misc/DateTime.h"
+
+namespace UE::Online {
+
+struct FQueryAchievementDefinitions
+{
+	static constexpr TCHAR Name[] = TEXT("QueryAchievementDefinitions");
+
+	struct Params
+	{
+		/** Local user performing the operation */
+		FOnlineAccountIdHandle LocalUserId;
+	};
+
+	struct Result
+	{
+	};
+};
+
+struct FGetAchievementIds
+{
+	static constexpr TCHAR Name[] = TEXT("GetAchievementIds");
+
+	struct Params
+	{
+		/** Local user performing the operation */
+		FOnlineAccountIdHandle LocalUserId;
+	};
+
+	struct Result
+	{
+		/** Available achievements */
+		TArray<FString> AchievementIds;
+	};
+};
+
+struct FAchievementStatDefinition
+{
+	/** Unique Id of the stat */
+	FString StatId;
+	/** Threshold value a user must meet for the achievement to auto unlock */
+	uint32 UnlockThreshold = 0;
+};
+
+FString ONLINESERVICESINTERFACE_API ToLogString(const FAchievementStatDefinition&);
+
+struct FAchievementDefinition
+{
+	/** Unique Id for this achievement */
+	FString AchievementId;
+	/** Localized display name of this achievement, once unlocked */
+	FText UnlockedDisplayName;
+	/** Localized description of this achievement, once unlocked */
+	FText UnlockedDescription;
+	/** Localized display name of this achievement, while locked */
+	FText LockedDisplayName;
+	/** Localized description of this achievement, while locked */
+	FText LockedDescription;
+	/** Localized flavor text */
+	FText FlavorText;
+	/** URL of the icon for this achievement, once unlocked */
+	FString UnlockedIconUrl;
+	/** URL of the icon for this achievement, while locked */
+	FString LockedIconUrl;
+	/** Is this achievement "hidden" until unlocked */
+	bool bIsHidden = false;
+	/** The stats that relate to this achievement */
+	TArray<FAchievementStatDefinition> StatDefinitions;
+};
+
+FString ONLINESERVICESINTERFACE_API ToLogString(const FAchievementDefinition&);
+
+struct FGetAchievementDefinition
+{
+	static constexpr TCHAR Name[] = TEXT("GetAchievementDefinition");
+
+	struct Params
+	{
+		/** Local user performing the operation */
+		FOnlineAccountIdHandle LocalUserId;
+		/** Achievement to get definition of */
+		FString AchievementId;
+	};
+
+	struct Result
+	{
+		/** Definition of the requested achievement */
+		FAchievementDefinition AchievementDefinition;
+	};
+};
+
+struct FQueryAchievementStates
+{
+	static constexpr TCHAR Name[] = TEXT("QueryAchievementStates");
+
+	struct Params
+	{
+		/** Local user performing the operation */
+		FOnlineAccountIdHandle LocalUserId;
+	};
+
+	struct Result
+	{
+	};
+};
+
+struct FAchievementState
+{
+	/** Achievement that this state relates to */
+	FString AchievementId;
+	/** Progress towards unlocking this achievement, as a percentage in the range [0.0, 1.0]. A value of 1.0 means the achievement is unlocked */
+	float Progress = 0.0f;
+	/** If unlocked, the time this achievement was unlocked */
+	FDateTime UnlockTime;
+};
+
+FString ONLINESERVICESINTERFACE_API ToLogString(const FAchievementState&);
+
+struct FGetAchievementState
+{
+	static constexpr TCHAR Name[] = TEXT("GetAchievementState");
+
+	struct Params
+	{
+		/** Local user performing the operation */
+		FOnlineAccountIdHandle LocalUserId;
+		/** Achievement to get state of */
+		FString AchievementId;
+	};
+
+	struct Result
+	{
+		/** The state of the achievement for the local user */
+		FAchievementState AchievementState;
+	};
+};
+
+struct FUnlockAchievements
+{
+	static constexpr TCHAR Name[] = TEXT("UnlockAchievements");
+
+	struct Params
+	{
+		/** Local user performing the operation */
+		FOnlineAccountIdHandle LocalUserId;
+		/** Achievements to unlock */
+		TArray<FString> AchievementIds;
+	};
+
+	struct Result
+	{
+	};
+};
+
+struct FDisplayAchievementUI
+{
+	static constexpr TCHAR Name[] = TEXT("DisplayAchievementUI");
+
+	struct Params
+	{
+		/** Local user performing the operation */
+		FOnlineAccountIdHandle LocalUserId;
+		/** Achievement to display */
+		FString AchievementId;
+	};
+
+	struct Result
+	{
+	};
+};
+
+/** Struct for AchievementStateUpdated event */
+struct FAchievementStateUpdated
+{
+	/** User whose achievement states have updated */
+	FOnlineAccountIdHandle LocalUserId;
+	/** Achievements which have updated */
+	TArray<FString> AchievementIds;
+};
+
+class IAchievements
+{
+public:
+	/**
+	 * Query all achievement definitions.
+	 */
+	virtual TOnlineAsyncOpHandle<FQueryAchievementDefinitions> QueryAchievementDefinitions(FQueryAchievementDefinitions::Params&& Params) = 0;
+
+	/**
+	 * Gets the achievement id's.
+	 * Requires first calling QueryAchievementDefinitions.
+	 */
+	virtual TOnlineResult<FGetAchievementIds> GetAchievementIds(FGetAchievementIds::Params&& Params) = 0;
+
+	/**
+	 * Gets an achievement definition by id.
+	 * Requires first calling QueryAchievementDefinitions.
+	 */
+	virtual TOnlineResult<FGetAchievementDefinition> GetAchievementDefinition(FGetAchievementDefinition::Params&& Params) = 0;
+
+	/**
+	 * Query the state of all achievements for the given player.
+	 * Requires first calling QueryAchievementDefinitions.
+	 */
+	virtual TOnlineAsyncOpHandle<FQueryAchievementStates> QueryAchievementStates(FQueryAchievementStates::Params&& Params) = 0;
+
+	/**
+	 * Gets the state of an achievement by id for the given player.
+	 * Requires first calling QueryAchievementStates.
+	 */
+	virtual TOnlineResult<FGetAchievementState> GetAchievementState(FGetAchievementState::Params&& Params) = 0;
+
+	/**
+	 * Manually unlock achievements
+	 * Requires first calling QueryAchievementStates.
+	 */
+	virtual TOnlineAsyncOpHandle<FUnlockAchievements> UnlockAchievements(FUnlockAchievements::Params&& Params) = 0;
+
+	/**
+	 * Launch the platform UI for a particular achievement
+	 */
+	virtual TOnlineResult<FDisplayAchievementUI> DisplayAchievementUI(FDisplayAchievementUI::Params&& Params) = 0;
+
+	/**
+	 * Event triggered when a player's achievement state changes
+	 */
+	virtual TOnlineEvent<void(const FAchievementStateUpdated&)> OnAchievementStateUpdated() = 0;
+};
+
+namespace Meta {
+
+BEGIN_ONLINE_STRUCT_META(FQueryAchievementDefinitions::Params)
+	ONLINE_STRUCT_FIELD(FQueryAchievementDefinitions::Params, LocalUserId)
+END_ONLINE_STRUCT_META()
+
+BEGIN_ONLINE_STRUCT_META(FQueryAchievementDefinitions::Result)
+END_ONLINE_STRUCT_META()
+
+BEGIN_ONLINE_STRUCT_META(FGetAchievementIds::Params)
+	ONLINE_STRUCT_FIELD(FGetAchievementIds::Params, LocalUserId)
+END_ONLINE_STRUCT_META()
+
+BEGIN_ONLINE_STRUCT_META(FGetAchievementIds::Result)
+	ONLINE_STRUCT_FIELD(FGetAchievementIds::Result, AchievementIds)
+END_ONLINE_STRUCT_META()
+
+BEGIN_ONLINE_STRUCT_META(FAchievementStatDefinition)
+	ONLINE_STRUCT_FIELD(FAchievementStatDefinition, StatId),
+	ONLINE_STRUCT_FIELD(FAchievementStatDefinition, UnlockThreshold)
+END_ONLINE_STRUCT_META()
+
+BEGIN_ONLINE_STRUCT_META(FAchievementDefinition)
+	ONLINE_STRUCT_FIELD(FAchievementDefinition, AchievementId),
+	ONLINE_STRUCT_FIELD(FAchievementDefinition, UnlockedDisplayName),
+	ONLINE_STRUCT_FIELD(FAchievementDefinition, UnlockedDescription),
+	ONLINE_STRUCT_FIELD(FAchievementDefinition, LockedDisplayName),
+	ONLINE_STRUCT_FIELD(FAchievementDefinition, LockedDescription),
+	ONLINE_STRUCT_FIELD(FAchievementDefinition, FlavorText),
+	ONLINE_STRUCT_FIELD(FAchievementDefinition, UnlockedIconUrl),
+	ONLINE_STRUCT_FIELD(FAchievementDefinition, LockedIconUrl),
+	ONLINE_STRUCT_FIELD(FAchievementDefinition, StatDefinitions)
+END_ONLINE_STRUCT_META()
+
+BEGIN_ONLINE_STRUCT_META(FGetAchievementDefinition::Params)
+	ONLINE_STRUCT_FIELD(FGetAchievementDefinition::Params, LocalUserId),
+	ONLINE_STRUCT_FIELD(FGetAchievementDefinition::Params, AchievementId)
+END_ONLINE_STRUCT_META()
+
+BEGIN_ONLINE_STRUCT_META(FGetAchievementDefinition::Result)
+	ONLINE_STRUCT_FIELD(FGetAchievementDefinition::Result, AchievementDefinition)
+END_ONLINE_STRUCT_META()
+
+BEGIN_ONLINE_STRUCT_META(FQueryAchievementStates::Params)
+	ONLINE_STRUCT_FIELD(FQueryAchievementStates::Params, LocalUserId)
+END_ONLINE_STRUCT_META()
+
+BEGIN_ONLINE_STRUCT_META(FQueryAchievementStates::Result)
+END_ONLINE_STRUCT_META()
+
+BEGIN_ONLINE_STRUCT_META(FAchievementState)
+	ONLINE_STRUCT_FIELD(FAchievementState, AchievementId),
+	ONLINE_STRUCT_FIELD(FAchievementState, Progress),
+	ONLINE_STRUCT_FIELD(FAchievementState, UnlockTime)
+END_ONLINE_STRUCT_META()
+
+BEGIN_ONLINE_STRUCT_META(FGetAchievementState::Params)
+	ONLINE_STRUCT_FIELD(FGetAchievementState::Params, LocalUserId),
+	ONLINE_STRUCT_FIELD(FGetAchievementState::Params, AchievementId)
+END_ONLINE_STRUCT_META()
+
+BEGIN_ONLINE_STRUCT_META(FGetAchievementState::Result)
+	ONLINE_STRUCT_FIELD(FGetAchievementState::Result, AchievementState)
+END_ONLINE_STRUCT_META()
+
+BEGIN_ONLINE_STRUCT_META(FUnlockAchievements::Params)
+	ONLINE_STRUCT_FIELD(FUnlockAchievements::Params, LocalUserId),
+	ONLINE_STRUCT_FIELD(FUnlockAchievements::Params, AchievementIds)
+END_ONLINE_STRUCT_META()
+
+BEGIN_ONLINE_STRUCT_META(FUnlockAchievements::Result)
+END_ONLINE_STRUCT_META()
+
+BEGIN_ONLINE_STRUCT_META(FDisplayAchievementUI::Params)
+	ONLINE_STRUCT_FIELD(FDisplayAchievementUI::Params, LocalUserId),
+	ONLINE_STRUCT_FIELD(FDisplayAchievementUI::Params, AchievementId)
+END_ONLINE_STRUCT_META()
+
+BEGIN_ONLINE_STRUCT_META(FDisplayAchievementUI::Result)
+END_ONLINE_STRUCT_META()
+
+BEGIN_ONLINE_STRUCT_META(FAchievementStateUpdated)
+	ONLINE_STRUCT_FIELD(FAchievementStateUpdated, LocalUserId),
+	ONLINE_STRUCT_FIELD(FAchievementStateUpdated, AchievementIds)
+END_ONLINE_STRUCT_META()
+
+/* Meta*/ }
+
+/* UE::Online */ }
diff --git a/Engine/Plugins/Online/Experimental/OnlineServices/Source/OnlineServicesInterface/Public/Online/AchievementsErrors.h b/Engine/Plugins/Online/Experimental/OnlineServices/Source/OnlineServicesInterface/Public/Online/AchievementsErrors.h
new file mode 100644
--- /dev/null
+++ b/Engine/Plugins/Online/Experimental/OnlineServices/Source/OnlineServicesInterface/Public/Online/AchievementsErrors.h
@@ -0,0 +1,15 @@
+// Copyright Epic Games, Inc. All Rights Reserved.
+
+#pragma once
+
+#include "Online/OnlineError.h"
+
+#define LOCTEXT_NAMESPACE "AchievementsErrors"
+
+namespace UE::Online::Errors {
+	UE_ONLINE_ERROR_CATEGORY(Achievements, Engine, 0x5, "Achievements")
+	UE_ONLINE_ERROR(Achievements, AlreadyUnlocked, 1, TEXT("AlreadyUnlocked"), LOCTEXT("AlreadyUnlocked", "That achievement is already unlocked"))
+
+/* UE::Online::Errors */ }
+
+#undef LOCTEXT_NAMESPACE
diff --git a/Engine/Plugins/Online/Experimental/OnlineServices/Source/OnlineServicesInterface/Public/Online/OnlineErrorDefinitions.h b/Engine/Plugins/Online/Experimental/OnlineServices/Source/OnlineServicesInterface/Public/Online/OnlineErrorDefinitions.h
--- a/Engine/Plugins/Online/Experimental/OnlineServices/Source/OnlineServicesInterface/Public/Online/OnlineErrorDefinitions.h
+++ b/Engine/Plugins/Online/Experimental/OnlineServices/Source/OnlineServicesInterface/Public/Online/OnlineErrorDefinitions.h
@@ -29,7 +29,8 @@ UE_ONLINE_ERROR_COMMON(Common, NotLoggedIn, 17, TEXT("not_logged_in"), LOCTEXT("
 UE_ONLINE_ERROR_COMMON(Common, NotFound, 18, TEXT("not_found"), LOCTEXT("NotFound", "Request not found"))
 UE_ONLINE_ERROR_COMMON(Common, WillRetry, 19, TEXT("will_retry"), LOCTEXT("WillRetry", "Retrying request"))
 UE_ONLINE_ERROR_COMMON(Common, Timeout, 20, TEXT("timeout"), LOCTEXT("Timeout", "Operation timed out"))
-UE_ONLINE_ERROR_COMMON(Common, Unknown, 21, TEXT("unknown"), LOCTEXT("Unknown", "Unknown Error"))
+UE_ONLINE_ERROR_COMMON(Common, InvalidState, 21, TEXT("invalid_state"), LOCTEXT("InvalidState", "Invalid state"))
+UE_ONLINE_ERROR_COMMON(Common, Unknown, 22, TEXT("unknown"), LOCTEXT("Unknown", "Unknown Error"))
 
 } /* namespace UE::Online::Errors */
 
diff --git a/Engine/Plugins/Online/Experimental/OnlineServices/Source/OnlineServicesInterface/Public/Online/OnlineServices.h b/Engine/Plugins/Online/Experimental/OnlineServices/Source/OnlineServicesInterface/Public/Online/OnlineServices.h
--- a/Engine/Plugins/Online/Experimental/OnlineServices/Source/OnlineServicesInterface/Public/Online/OnlineServices.h
+++ b/Engine/Plugins/Online/Experimental/OnlineServices/Source/OnlineServicesInterface/Public/Online/OnlineServices.h
@@ -13,6 +13,7 @@
 
 // Interfaces
 using IOnlineServicesPtr = TSharedPtr<class IOnlineServices>;
+using IAchievementsPtr = TSharedPtr<class IAchievements>;
 using IAuthPtr = TSharedPtr<class IAuth>;
 using IFriendsPtr = TSharedPtr<class IFriends>;
 using IPresencePtr = TSharedPtr<class IPresence>;
@@ -56,6 +57,11 @@
 	/**
 	 *
 	 */
+	virtual IAchievementsPtr GetAchievementsInterface() = 0;
+
+	/**
+	 *
+	 */
 	virtual IAuthPtr GetAuthInterface() = 0;
 
 	/**
