- TitleFile interface and common impl

diff --git a/Engine/Plugins/Online/Experimental/OnlineServices/Source/OnlineServicesCommon/Private/Online/OnlineServicesCommon.cpp b/Engine/Plugins/Online/Experimental/OnlineServices/Source/OnlineServicesCommon/Private/Online/OnlineServicesCommon.cpp
--- a/Engine/Plugins/Online/Experimental/OnlineServices/Source/OnlineServicesCommon/Private/Online/OnlineServicesCommon.cpp
+++ b/Engine/Plugins/Online/Experimental/OnlineServices/Source/OnlineServicesCommon/Private/Online/OnlineServicesCommon.cpp
@@ -85,6 +85,11 @@ IPrivilegesPtr FOnlineServicesCommon::GetPrivilegesInterface()
 	return IPrivilegesPtr(AsShared(), Get<IPrivileges>());
 }
 
+ITitleFilePtr FOnlineServicesCommon::GetTitleFileInterface()
+{
+	return ITitleFilePtr(AsShared(), Get<ITitleFile>());
+}
+
 TOnlineResult<FGetResolvedConnectString> FOnlineServicesCommon::GetResolvedConnectString(FGetResolvedConnectString::Params&& Params)
 {
 	return TOnlineResult<FGetResolvedConnectString>(Errors::NotImplemented());
diff --git a/Engine/Plugins/Online/Experimental/OnlineServices/Source/OnlineServicesCommon/Private/Online/TitleFileCommon.cpp b/Engine/Plugins/Online/Experimental/OnlineServices/Source/OnlineServicesCommon/Private/Online/TitleFileCommon.cpp
new file mode 100644
--- /dev/null
+++ b/Engine/Plugins/Online/Experimental/OnlineServices/Source/OnlineServicesCommon/Private/Online/TitleFileCommon.cpp
@@ -0,0 +1,42 @@
+// Copyright Epic Games, Inc. All Rights Reserved.
+
+#include "Online/TitleFileCommon.h"
+
+#include "Online/OnlineAsyncOp.h"
+#include "Online/OnlineErrorDefinitions.h"
+#include "Online/OnlineServicesCommon.h"
+
+namespace UE::Online {
+
+FTitleFileCommon::FTitleFileCommon(FOnlineServicesCommon& InServices)
+	: TOnlineComponent(TEXT("TitleFile"), InServices)
+{
+}
+
+void FTitleFileCommon::RegisterCommands()
+{
+	RegisterCommand(&FTitleFileCommon::EnumerateFiles);
+	RegisterCommand(&FTitleFileCommon::GetEnumeratedFiles);
+	RegisterCommand(&FTitleFileCommon::ReadFile);
+}
+
+TOnlineAsyncOpHandle<FTitleFileEnumerateFiles> FTitleFileCommon::EnumerateFiles(FTitleFileEnumerateFiles::Params&& Params)
+{
+	TOnlineAsyncOpRef<FTitleFileEnumerateFiles> Operation = GetOp<FTitleFileEnumerateFiles>(MoveTemp(Params));
+	Operation->SetError(Errors::NotImplemented());
+	return Operation->GetHandle();
+}
+
+TOnlineResult<FTitleFileGetEnumeratedFiles> FTitleFileCommon::GetEnumeratedFiles(FTitleFileGetEnumeratedFiles::Params&& Params)
+{
+	return TOnlineResult<FTitleFileGetEnumeratedFiles>(Errors::NotImplemented());
+}
+
+TOnlineAsyncOpHandle<FTitleFileReadFile> FTitleFileCommon::ReadFile(FTitleFileReadFile::Params&& Params)
+{
+	TOnlineAsyncOpRef<FTitleFileReadFile> Operation = GetOp<FTitleFileReadFile>(MoveTemp(Params));
+	Operation->SetError(Errors::NotImplemented());
+	return Operation->GetHandle();
+}
+
+/* UE::Online */}
\ No newline at end of file
diff --git a/Engine/Plugins/Online/Experimental/OnlineServices/Source/OnlineServicesCommon/Public/Online/OnlineServicesCommon.h b/Engine/Plugins/Online/Experimental/OnlineServices/Source/OnlineServicesCommon/Public/Online/OnlineServicesCommon.h
--- a/Engine/Plugins/Online/Experimental/OnlineServices/Source/OnlineServicesCommon/Public/Online/OnlineServicesCommon.h
+++ b/Engine/Plugins/Online/Experimental/OnlineServices/Source/OnlineServicesCommon/Public/Online/OnlineServicesCommon.h
@@ -44,6 +44,7 @@ public:
 	virtual ISessionsPtr GetSessionsInterface() override;
 	virtual IConnectivityPtr GetConnectivityInterface() override;
 	virtual IPrivilegesPtr GetPrivilegesInterface() override;
+	virtual ITitleFilePtr GetTitleFileInterface() override;
 	virtual TOnlineResult<FGetResolvedConnectString> GetResolvedConnectString(FGetResolvedConnectString::Params&& Params) override;
 	virtual FName GetInstanceName() const override;
 
diff --git a/Engine/Plugins/Online/Experimental/OnlineServices/Source/OnlineServicesCommon/Public/Online/TitleFileCommon.h b/Engine/Plugins/Online/Experimental/OnlineServices/Source/OnlineServicesCommon/Public/Online/TitleFileCommon.h
new file mode 100644
--- /dev/null
+++ b/Engine/Plugins/Online/Experimental/OnlineServices/Source/OnlineServicesCommon/Public/Online/TitleFileCommon.h
@@ -0,0 +1,28 @@
+// Copyright Epic Games, Inc. All Rights Reserved.
+
+#pragma once
+
+#include "Online/TitleFile.h"
+#include "Online/OnlineComponent.h"
+
+namespace UE::Online {
+
+class FOnlineServicesCommon;
+
+class ONLINESERVICESCOMMON_API FTitleFileCommon : public TOnlineComponent<ITitleFile>
+{
+public:
+	using Super = ITitleFile;
+
+	FTitleFileCommon(FOnlineServicesCommon& InServices);
+
+	// TOnlineComponent
+	virtual void RegisterCommands() override;
+
+	// ITitleFile
+	virtual TOnlineAsyncOpHandle<FTitleFileEnumerateFiles> EnumerateFiles(FTitleFileEnumerateFiles::Params&& Params) override;
+	virtual TOnlineResult<FTitleFileGetEnumeratedFiles> GetEnumeratedFiles(FTitleFileGetEnumeratedFiles::Params&& Params) override;
+	virtual TOnlineAsyncOpHandle<FTitleFileReadFile> ReadFile(FTitleFileReadFile::Params&& Params) override;
+};
+
+/* UE::Online */}
\ No newline at end of file
diff --git a/Engine/Plugins/Online/Experimental/OnlineServices/Source/OnlineServicesInterface/Public/Online/OnlineServices.h b/Engine/Plugins/Online/Experimental/OnlineServices/Source/OnlineServicesInterface/Public/Online/OnlineServices.h
--- a/Engine/Plugins/Online/Experimental/OnlineServices/Source/OnlineServicesInterface/Public/Online/OnlineServices.h
+++ b/Engine/Plugins/Online/Experimental/OnlineServices/Source/OnlineServicesInterface/Public/Online/OnlineServices.h
@@ -22,6 +22,7 @@ using ILobbiesPtr = TSharedPtr<class ILobbies>;
 using ISessionsPtr = TSharedPtr<class ISessions>;
 using IConnectivityPtr = TSharedPtr<class IConnectivity>;
 using IPrivilegesPtr = TSharedPtr<class IPrivileges>;
+using ITitleFilePtr = TSharedPtr<class ITitleFile>;
 
 struct FGetResolvedConnectString
 {
@@ -101,6 +102,11 @@ public:
 	 */
 	virtual IPrivilegesPtr GetPrivilegesInterface() = 0;
 
+	/**
+	 *
+	 */
+	virtual ITitleFilePtr GetTitleFileInterface() = 0;
+
 	/** 
 	 * Get the connectivity string used for client travel
 	 */
diff --git a/Engine/Plugins/Online/Experimental/OnlineServices/Source/OnlineServicesInterface/Public/Online/TitleFile.h b/Engine/Plugins/Online/Experimental/OnlineServices/Source/OnlineServicesInterface/Public/Online/TitleFile.h
new file mode 100644
--- /dev/null
+++ b/Engine/Plugins/Online/Experimental/OnlineServices/Source/OnlineServicesInterface/Public/Online/TitleFile.h
@@ -0,0 +1,113 @@
+// Copyright Epic Games, Inc. All Rights Reserved.
+
+#pragma once
+
+#include "Online/CoreOnline.h"
+#include "Online/OnlineAsyncOpHandle.h"
+#include "Online/OnlineMeta.h"
+
+namespace UE::Online {
+
+struct FTitleFileEnumerateFiles
+{
+	static constexpr TCHAR Name[] = TEXT("EnumerateFiles");
+
+	struct Params
+	{
+		/** Local user performing the operation */
+		FOnlineAccountIdHandle LocalUserId;
+	};
+
+	struct Result
+	{
+	};
+};
+
+struct FTitleFileGetEnumeratedFiles
+{
+	static constexpr TCHAR Name[] = TEXT("GetEnumeratedFiles");
+
+	struct Params
+	{
+		/** Local user performing the operation */
+		FOnlineAccountIdHandle LocalUserId;
+	};
+
+	struct Result
+	{
+		/** Available files */
+		TArray<FString> Filenames;
+	};
+};
+
+using FTitleFileContents = TArray<uint8>;
+using FTitleFileContentsRef = TSharedRef<const FTitleFileContents>;
+
+struct FTitleFileReadFile
+{
+	static constexpr TCHAR Name[] = TEXT("ReadFile");
+
+	struct Params
+	{
+		/** Local user performing the operation */
+		FOnlineAccountIdHandle LocalUserId;
+		/** File to be opened */
+		FString Filename;
+	};
+
+	struct Result
+	{
+		/** Contents of the file */
+		FTitleFileContentsRef FileContents;
+	};
+};
+
+class ITitleFile
+{
+public:
+	/**
+	 * Enumerate all available files
+	 */
+	virtual TOnlineAsyncOpHandle<FTitleFileEnumerateFiles> EnumerateFiles(FTitleFileEnumerateFiles::Params&& Params) = 0;
+
+	/**
+	 * Get the cached list of files enumerated by previous call to EnumerateFiles
+	 */
+	virtual TOnlineResult<FTitleFileGetEnumeratedFiles> GetEnumeratedFiles(FTitleFileGetEnumeratedFiles::Params&& Params) = 0;
+
+	/**
+	 * Read a file and return the file contents
+	 */
+	virtual TOnlineAsyncOpHandle<FTitleFileReadFile> ReadFile(FTitleFileReadFile::Params&& Params) = 0;
+};
+
+
+namespace Meta {
+
+BEGIN_ONLINE_STRUCT_META(FTitleFileEnumerateFiles::Params)
+	ONLINE_STRUCT_FIELD(FTitleFileEnumerateFiles::Params, LocalUserId)
+END_ONLINE_STRUCT_META()
+
+BEGIN_ONLINE_STRUCT_META(FTitleFileEnumerateFiles::Result)
+END_ONLINE_STRUCT_META()
+
+BEGIN_ONLINE_STRUCT_META(FTitleFileGetEnumeratedFiles::Params)
+	ONLINE_STRUCT_FIELD(FTitleFileGetEnumeratedFiles::Params, LocalUserId)
+END_ONLINE_STRUCT_META()
+
+BEGIN_ONLINE_STRUCT_META(FTitleFileGetEnumeratedFiles::Result)
+	ONLINE_STRUCT_FIELD(FTitleFileGetEnumeratedFiles::Result, Filenames)
+END_ONLINE_STRUCT_META()
+
+BEGIN_ONLINE_STRUCT_META(FTitleFileReadFile::Params)
+	ONLINE_STRUCT_FIELD(FTitleFileReadFile::Params, LocalUserId),
+	ONLINE_STRUCT_FIELD(FTitleFileReadFile::Params, Filename)
+END_ONLINE_STRUCT_META()
+
+BEGIN_ONLINE_STRUCT_META(FTitleFileReadFile::Result)
+	ONLINE_STRUCT_FIELD(FTitleFileReadFile::Result, FileContents)
+END_ONLINE_STRUCT_META()
+
+/* Meta*/ }
+
+/* UE::Online */ }
