OSSv2: Add Stats interface and common implementation.

diff --git a/Engine/Plugins/Online/Experimental/OnlineServices/Source/OnlineServicesCommon/Private/Online/OnlineServicesCommon.cpp b/Engine/Plugins/Online/Experimental/OnlineServices/Source/OnlineServicesCommon/Private/Online/OnlineServicesCommon.cpp
--- a/Engine/Plugins/Online/Experimental/OnlineServices/Source/OnlineServicesCommon/Private/Online/OnlineServicesCommon.cpp
+++ b/Engine/Plugins/Online/Experimental/OnlineServices/Source/OnlineServicesCommon/Private/Online/OnlineServicesCommon.cpp
@@ -75,6 +75,11 @@ ISessionsPtr FOnlineServicesCommon::GetSessionsInterface()
 	return ISessionsPtr(AsShared(), Get<ISessions>());
 }
 
+IStatsPtr FOnlineServicesCommon::GetStatsInterface()
+{
+	return IStatsPtr(AsShared(), Get<IStats>());
+}
+
 IConnectivityPtr FOnlineServicesCommon::GetConnectivityInterface()
 {
 	return IConnectivityPtr(AsShared(), Get<IConnectivity>());
diff --git a/Engine/Plugins/Online/Experimental/OnlineServices/Source/OnlineServicesCommon/Private/Online/StatsCommon.cpp b/Engine/Plugins/Online/Experimental/OnlineServices/Source/OnlineServicesCommon/Private/Online/StatsCommon.cpp
new file mode 100644
--- /dev/null
+++ b/Engine/Plugins/Online/Experimental/OnlineServices/Source/OnlineServicesCommon/Private/Online/StatsCommon.cpp
@@ -0,0 +1,62 @@
+// Copyright Epic Games, Inc. All Rights Reserved.
+
+#include "Online/StatsCommon.h"
+
+#include "Online/Auth.h"
+#include "Online/OnlineServicesCommon.h"
+
+namespace UE::Online {
+
+FStatsCommon::FStatsCommon(FOnlineServicesCommon& InServices)
+	: TOnlineComponent(TEXT("Stats"), InServices)
+{
+}
+
+void FStatsCommon::Initialize()
+{
+	TOnlineComponent<IStats>::Initialize();
+}
+
+void FStatsCommon::RegisterCommands()
+{
+	TOnlineComponent<IStats>::RegisterCommands();
+
+	RegisterCommand(&FStatsCommon::UpdateStats);
+	RegisterCommand(&FStatsCommon::QueryStats);
+	RegisterCommand(&FStatsCommon::BatchQueryStats);
+#if !UE_BUILD_SHIPPING
+	RegisterCommand(&FStatsCommon::ResetStats);
+#endif // !UE_BUILD_SHIPPING
+}
+
+TOnlineAsyncOpHandle<FUpdateStats> FStatsCommon::UpdateStats(FUpdateStats::Params&& Params)
+{
+	TOnlineAsyncOpRef<FUpdateStats> Operation = GetOp<FUpdateStats>(MoveTemp(Params));
+	Operation->SetError(Errors::NotImplemented());
+	return Operation->GetHandle();
+}
+
+TOnlineAsyncOpHandle<FQueryStats> FStatsCommon::QueryStats(FQueryStats::Params&& Params)
+{
+	TOnlineAsyncOpRef<FQueryStats> Operation = GetOp<FQueryStats>(MoveTemp(Params));
+	Operation->SetError(Errors::NotImplemented());
+	return Operation->GetHandle();
+}
+
+TOnlineAsyncOpHandle<FBatchQueryStats> FStatsCommon::BatchQueryStats(FBatchQueryStats::Params&& Params)
+{
+	TOnlineAsyncOpRef<FBatchQueryStats> Operation = GetOp<FBatchQueryStats>(MoveTemp(Params));
+	Operation->SetError(Errors::NotImplemented());
+	return Operation->GetHandle();
+}
+
+#if !UE_BUILD_SHIPPING
+TOnlineAsyncOpHandle<FResetStats> FStatsCommon::ResetStats(FResetStats::Params&& Params)
+{
+	TOnlineAsyncOpRef<FResetStats> Operation = GetOp<FResetStats>(MoveTemp(Params));
+	Operation->SetError(Errors::NotImplemented());
+	return Operation->GetHandle();
+}
+#endif // !UE_BUILD_SHIPPING
+
+/* UE::Online */ }
diff --git a/Engine/Plugins/Online/Experimental/OnlineServices/Source/OnlineServicesCommon/Public/Online/OnlineServicesCommon.h b/Engine/Plugins/Online/Experimental/OnlineServices/Source/OnlineServicesCommon/Public/Online/OnlineServicesCommon.h
--- a/Engine/Plugins/Online/Experimental/OnlineServices/Source/OnlineServicesCommon/Public/Online/OnlineServicesCommon.h
+++ b/Engine/Plugins/Online/Experimental/OnlineServices/Source/OnlineServicesCommon/Public/Online/OnlineServicesCommon.h
@@ -42,6 +42,7 @@ public:
 	virtual IExternalUIPtr GetExternalUIInterface() override;
 	virtual ILobbiesPtr GetLobbiesInterface() override;
 	virtual ISessionsPtr GetSessionsInterface() override;
+	virtual IStatsPtr GetStatsInterface() override;
 	virtual IConnectivityPtr GetConnectivityInterface() override;
 	virtual IPrivilegesPtr GetPrivilegesInterface() override;
 	virtual ITitleFilePtr GetTitleFileInterface() override;
diff --git a/Engine/Plugins/Online/Experimental/OnlineServices/Source/OnlineServicesCommon/Public/Online/StatsCommon.h b/Engine/Plugins/Online/Experimental/OnlineServices/Source/OnlineServicesCommon/Public/Online/StatsCommon.h
new file mode 100644
--- /dev/null
+++ b/Engine/Plugins/Online/Experimental/OnlineServices/Source/OnlineServicesCommon/Public/Online/StatsCommon.h
@@ -0,0 +1,30 @@
+// Copyright Epic Games, Inc. All Rights Reserved.
+
+#pragma once
+
+#include "Online/Stats.h"
+#include "Online/OnlineComponent.h"
+
+namespace UE::Online {
+
+	class ONLINESERVICESCOMMON_API FStatsCommon : public TOnlineComponent<IStats>
+	{
+	public:
+		using Super = IStats;
+
+		FStatsCommon(FOnlineServicesCommon& InServices);
+
+		// TOnlineComponent
+		virtual void Initialize() override;
+		virtual void RegisterCommands() override;
+
+		// IStats
+		virtual TOnlineAsyncOpHandle<FUpdateStats> UpdateStats(FUpdateStats::Params&& Params) override;
+		virtual TOnlineAsyncOpHandle<FQueryStats> QueryStats(FQueryStats::Params&& Params) override;
+		virtual TOnlineAsyncOpHandle<FBatchQueryStats> BatchQueryStats(FBatchQueryStats::Params&& Params) override;
+#if !UE_BUILD_SHIPPING
+		virtual TOnlineAsyncOpHandle<FResetStats> ResetStats(FResetStats::Params&& Params) override;
+#endif // !UE_BUILD_SHIPPING
+	};
+
+/* UE::Online */ }
diff --git a/Engine/Plugins/Online/Experimental/OnlineServices/Source/OnlineServicesInterface/Public/Online/OnlineServices.h b/Engine/Plugins/Online/Experimental/OnlineServices/Source/OnlineServicesInterface/Public/Online/OnlineServices.h
--- a/Engine/Plugins/Online/Experimental/OnlineServices/Source/OnlineServicesInterface/Public/Online/OnlineServices.h
+++ b/Engine/Plugins/Online/Experimental/OnlineServices/Source/OnlineServicesInterface/Public/Online/OnlineServices.h
@@ -20,6 +20,7 @@ using IPresencePtr = TSharedPtr<class IPresence>;
 using IExternalUIPtr = TSharedPtr<class IExternalUI>;
 using ILobbiesPtr = TSharedPtr<class ILobbies>;
 using ISessionsPtr = TSharedPtr<class ISessions>;
+using IStatsPtr = TSharedPtr<class IStats>;
 using IConnectivityPtr = TSharedPtr<class IConnectivity>;
 using IPrivilegesPtr = TSharedPtr<class IPrivileges>;
 using ITitleFilePtr = TSharedPtr<class ITitleFile>;
@@ -92,6 +93,12 @@ public:
 	 */
 	virtual ISessionsPtr GetSessionsInterface() = 0;
 
+	/**
+	 * Get the stats implementation
+	 * @return stats implementation, may be null if not implemented for this service
+	 */
+	virtual IStatsPtr GetStatsInterface() = 0;
+
 	/**
 	 * 
 	 */
diff --git a/Engine/Plugins/Online/Experimental/OnlineServices/Source/OnlineServicesInterface/Public/Online/OnlineUtils.h b/Engine/Plugins/Online/Experimental/OnlineServices/Source/OnlineServicesInterface/Public/Online/OnlineUtils.h
--- a/Engine/Plugins/Online/Experimental/OnlineServices/Source/OnlineServicesInterface/Public/Online/OnlineUtils.h
+++ b/Engine/Plugins/Online/Experimental/OnlineServices/Source/OnlineServicesInterface/Public/Online/OnlineUtils.h
@@ -23,6 +23,8 @@
 inline FString ToLogString(int32 Value);
 inline FString ToLogString(uint64 Value);
 inline FString ToLogString(int64 Value);
+inline FString ToLogString(float Value);
+inline FString ToLogString(double Value);
 inline FString ToLogString(bool Value);
 template <typename T> std::enable_if_t<!TModels<Meta::COnlineMetadataAvailable, T>::Value, FString> ToLogString(const T& Value);
 template <typename T> std::enable_if_t<TModels<Meta::COnlineMetadataAvailable, T>::Value, FString> ToLogString(const T& Value);
@@ -159,6 +161,16 @@
 	return FString::Printf(TEXT("%lli"), Value);
 }
 
+inline FString ToLogString(float Value)
+{
+	return FString::Printf(TEXT("%.2f"), Value);
+}
+
+inline FString ToLogString(double Value)
+{
+	return FString::Printf(TEXT("%.2f"), Value);
+}
+
 inline FString ToLogString(bool Value)
 {
 	return ::LexToString(Value);
diff --git a/Engine/Plugins/Online/Experimental/OnlineServices/Source/OnlineServicesInterface/Public/Online/Stats.h b/Engine/Plugins/Online/Experimental/OnlineServices/Source/OnlineServicesInterface/Public/Online/Stats.h
new file mode 100644
--- /dev/null
+++ b/Engine/Plugins/Online/Experimental/OnlineServices/Source/OnlineServicesInterface/Public/Online/Stats.h
@@ -0,0 +1,177 @@
+// Copyright Epic Games, Inc. All Rights Reserved.
+
+#pragma once
+
+#include "Online/CoreOnline.h"
+#include "Online/OnlineAsyncOpHandle.h"
+#include "Online/OnlineMeta.h"
+
+namespace UE::Online {
+
+class FOnlineError;
+
+//using FStatValue = SchemaVariant;
+using FStatValue = TVariant<bool, int32, int64, float, double, FString>;
+
+struct FUserStats
+{
+	/* The user id */
+	FOnlineAccountIdHandle UserId;
+	/* The stats of one account */
+	TMap<FString, FStatValue> Stats;
+};
+
+struct FUpdateStats
+{
+	static constexpr TCHAR Name[] = TEXT("UpdateStats");
+
+	struct Params
+	{
+		/* Local user id */
+		FOnlineAccountIdHandle LocalUserId;
+		/* The stats of users to update */
+		TArray<FUserStats> UpdateUsersStats;
+	};
+
+	struct Result
+	{
+	};
+};
+
+struct FQueryStats
+{
+	static constexpr TCHAR Name[] = TEXT("QueryStats");
+
+	struct Params
+	{
+		/* Local user id */
+		FOnlineAccountIdHandle LocalUserId;
+		/* Target user id to query */
+		FOnlineAccountIdHandle TargetUserId;
+	};
+
+	struct Result
+	{
+		/* The stats of queried user */
+		TMap<FString, FStatValue> Stats;
+	};
+};
+
+struct FBatchQueryStats
+{
+	static constexpr TCHAR Name[] = TEXT("BatchQueryStats");
+
+	struct Params
+	{
+		/* Local user id */
+		FOnlineAccountIdHandle LocalUserId;
+		/* The user ids to query */
+		TArray<FOnlineAccountIdHandle> TargetUserIds;
+	};
+
+	struct Result
+	{
+		/* The stats of queried users */
+		TArray<FUserStats> UsersStats;
+	};
+};
+
+#if !UE_BUILD_SHIPPING
+struct FResetStats
+{
+	static constexpr TCHAR Name[] = TEXT("ResetStats");
+
+	struct Params
+	{
+		/* Local user id */
+		FOnlineAccountIdHandle LocalUserId;
+	};
+
+	struct Result
+	{
+	};
+};
+#endif // !UE_BUILD_SHIPPING
+
+// Stats interface is used to upload the stats to EOS or first party platforms, to complete corresponding InGame/Platform features such as Stats Query, Achievements, Leaderboard, etc.
+class IStats
+{
+public:
+	/**
+	 * Upload/send stats to the backend/platform.
+	 *
+	 * @param Parameters for the UpdateStats call
+	 * @return
+	 */
+	virtual TOnlineAsyncOpHandle<FUpdateStats> UpdateStats(FUpdateStats::Params&& Params) = 0;
+	/**
+	 * Get stats of the specified user.
+	 *
+	 * @param Parameters for the QueryStats call
+	 * @return
+	 */
+	virtual TOnlineAsyncOpHandle<FQueryStats> QueryStats(FQueryStats::Params&& Params) = 0;
+	/**
+	 * Get stats of the specified users.
+	 *
+	 * @param Parameters for the BatchQueryStats call
+	 * @return
+	 */
+	virtual TOnlineAsyncOpHandle<FBatchQueryStats> BatchQueryStats(FBatchQueryStats::Params&& Params) = 0;
+
+#if !UE_BUILD_SHIPPING
+	/**
+	 * Reset stats of the local user.
+	 *
+	 * @param Parameters for the ResetStats call
+	 * @return
+	 */
+	virtual TOnlineAsyncOpHandle<FResetStats> ResetStats(FResetStats::Params&& Params) = 0;
+#endif // !UE_BUILD_SHIPPING
+};
+
+namespace Meta {
+
+BEGIN_ONLINE_STRUCT_META(FUserStats)
+	ONLINE_STRUCT_FIELD(FUserStats, UserId),
+	ONLINE_STRUCT_FIELD(FUserStats, Stats)
+END_ONLINE_STRUCT_META()
+
+BEGIN_ONLINE_STRUCT_META(FUpdateStats::Params)
+	ONLINE_STRUCT_FIELD(FUpdateStats::Params, LocalUserId),
+	ONLINE_STRUCT_FIELD(FUpdateStats::Params, UpdateUsersStats)
+END_ONLINE_STRUCT_META()
+
+BEGIN_ONLINE_STRUCT_META(FUpdateStats::Result)
+END_ONLINE_STRUCT_META()
+
+BEGIN_ONLINE_STRUCT_META(FQueryStats::Params)
+	ONLINE_STRUCT_FIELD(FQueryStats::Params, LocalUserId),
+	ONLINE_STRUCT_FIELD(FQueryStats::Params, TargetUserId)
+END_ONLINE_STRUCT_META()
+
+BEGIN_ONLINE_STRUCT_META(FQueryStats::Result)
+	ONLINE_STRUCT_FIELD(FQueryStats::Result, Stats)
+END_ONLINE_STRUCT_META()
+
+BEGIN_ONLINE_STRUCT_META(FBatchQueryStats::Params)
+	ONLINE_STRUCT_FIELD(FBatchQueryStats::Params, LocalUserId),
+	ONLINE_STRUCT_FIELD(FBatchQueryStats::Params, TargetUserIds)
+END_ONLINE_STRUCT_META()
+
+BEGIN_ONLINE_STRUCT_META(FBatchQueryStats::Result)
+	ONLINE_STRUCT_FIELD(FBatchQueryStats::Result, UsersStats)
+END_ONLINE_STRUCT_META()
+
+#if !UE_BUILD_SHIPPING
+BEGIN_ONLINE_STRUCT_META(FResetStats::Params)
+	ONLINE_STRUCT_FIELD(FResetStats::Params, LocalUserId)
+END_ONLINE_STRUCT_META()
+
+BEGIN_ONLINE_STRUCT_META(FResetStats::Result)
+END_ONLINE_STRUCT_META()
+#endif // !UE_BUILD_SHIPPING
+
+/* Meta*/ }
+
+/* UE::Online */ }
