- UserFile interface and common impl

diff --git a/Engine/Plugins/Online/Experimental/OnlineServices/Source/OnlineServicesCommon/Private/Online/OnlineServicesCommon.cpp b/Engine/Plugins/Online/Experimental/OnlineServices/Source/OnlineServicesCommon/Private/Online/OnlineServicesCommon.cpp
--- a/Engine/Plugins/Online/Experimental/OnlineServices/Source/OnlineServicesCommon/Private/Online/OnlineServicesCommon.cpp
+++ b/Engine/Plugins/Online/Experimental/OnlineServices/Source/OnlineServicesCommon/Private/Online/OnlineServicesCommon.cpp
@@ -95,6 +95,11 @@ ITitleFilePtr FOnlineServicesCommon::GetTitleFileInterface()
 	return ITitleFilePtr(AsShared(), Get<ITitleFile>());
 }
 
+IUserFilePtr FOnlineServicesCommon::GetUserFileInterface()
+{
+	return IUserFilePtr(AsShared(), Get<IUserFile>());
+}
+
 TOnlineResult<FGetResolvedConnectString> FOnlineServicesCommon::GetResolvedConnectString(FGetResolvedConnectString::Params&& Params)
 {
 	return TOnlineResult<FGetResolvedConnectString>(Errors::NotImplemented());
diff --git a/Engine/Plugins/Online/Experimental/OnlineServices/Source/OnlineServicesCommon/Private/Online/UserFileCommon.cpp b/Engine/Plugins/Online/Experimental/OnlineServices/Source/OnlineServicesCommon/Private/Online/UserFileCommon.cpp
new file mode 100644
--- /dev/null
+++ b/Engine/Plugins/Online/Experimental/OnlineServices/Source/OnlineServicesCommon/Private/Online/UserFileCommon.cpp
@@ -0,0 +1,66 @@
+// Copyright Epic Games, Inc. All Rights Reserved.
+
+#include "Online/UserFileCommon.h"
+
+#include "Online/OnlineAsyncOp.h"
+#include "Online/OnlineErrorDefinitions.h"
+#include "Online/OnlineServicesCommon.h"
+
+namespace UE::Online {
+
+FUserFileCommon::FUserFileCommon(FOnlineServicesCommon& InServices)
+	: TOnlineComponent(TEXT("UserFile"), InServices)
+{
+}
+
+void FUserFileCommon::RegisterCommands()
+{
+	RegisterCommand(&FUserFileCommon::EnumerateFiles);
+	RegisterCommand(&FUserFileCommon::GetEnumeratedFiles);
+	RegisterCommand(&FUserFileCommon::ReadFile);
+	RegisterCommand(&FUserFileCommon::WriteFile);
+	RegisterCommand(&FUserFileCommon::CopyFile);
+	RegisterCommand(&FUserFileCommon::DeleteFile);
+}
+
+TOnlineAsyncOpHandle<FUserFileEnumerateFiles> FUserFileCommon::EnumerateFiles(FUserFileEnumerateFiles::Params&& Params)
+{
+	TOnlineAsyncOpRef<FUserFileEnumerateFiles> Operation = GetOp<FUserFileEnumerateFiles>(MoveTemp(Params));
+	Operation->SetError(Errors::NotImplemented());
+	return Operation->GetHandle();
+}
+
+TOnlineResult<FUserFileGetEnumeratedFiles> FUserFileCommon::GetEnumeratedFiles(FUserFileGetEnumeratedFiles::Params&& Params)
+{
+	return TOnlineResult<FUserFileGetEnumeratedFiles>(Errors::NotImplemented());
+}
+
+TOnlineAsyncOpHandle<FUserFileReadFile> FUserFileCommon::ReadFile(FUserFileReadFile::Params&& Params)
+{
+	TOnlineAsyncOpRef<FUserFileReadFile> Operation = GetOp<FUserFileReadFile>(MoveTemp(Params));
+	Operation->SetError(Errors::NotImplemented());
+	return Operation->GetHandle();
+}
+
+TOnlineAsyncOpHandle<FUserFileWriteFile> FUserFileCommon::WriteFile(FUserFileWriteFile::Params&& Params)
+{
+	TOnlineAsyncOpRef<FUserFileWriteFile> Operation = GetOp<FUserFileWriteFile>(MoveTemp(Params));
+	Operation->SetError(Errors::NotImplemented());
+	return Operation->GetHandle();
+}
+
+TOnlineAsyncOpHandle<FUserFileCopyFile> FUserFileCommon::CopyFile(FUserFileCopyFile::Params&& Params)
+{
+	TOnlineAsyncOpRef<FUserFileCopyFile> Operation = GetOp<FUserFileCopyFile>(MoveTemp(Params));
+	Operation->SetError(Errors::NotImplemented());
+	return Operation->GetHandle();
+}
+
+TOnlineAsyncOpHandle<FUserFileDeleteFile> FUserFileCommon::DeleteFile(FUserFileDeleteFile::Params&& Params)
+{
+	TOnlineAsyncOpRef<FUserFileDeleteFile> Operation = GetOp<FUserFileDeleteFile>(MoveTemp(Params));
+	Operation->SetError(Errors::NotImplemented());
+	return Operation->GetHandle();
+}
+
+/* UE::Online */}
\ No newline at end of file
diff --git a/Engine/Plugins/Online/Experimental/OnlineServices/Source/OnlineServicesCommon/Public/Online/OnlineServicesCommon.h b/Engine/Plugins/Online/Experimental/OnlineServices/Source/OnlineServicesCommon/Public/Online/OnlineServicesCommon.h
--- a/Engine/Plugins/Online/Experimental/OnlineServices/Source/OnlineServicesCommon/Public/Online/OnlineServicesCommon.h
+++ b/Engine/Plugins/Online/Experimental/OnlineServices/Source/OnlineServicesCommon/Public/Online/OnlineServicesCommon.h
@@ -46,6 +46,7 @@ public:
 	virtual IConnectivityPtr GetConnectivityInterface() override;
 	virtual IPrivilegesPtr GetPrivilegesInterface() override;
 	virtual ITitleFilePtr GetTitleFileInterface() override;
+	virtual IUserFilePtr GetUserFileInterface() override;
 	virtual TOnlineResult<FGetResolvedConnectString> GetResolvedConnectString(FGetResolvedConnectString::Params&& Params) override;
 	virtual FName GetInstanceName() const override;
 
diff --git a/Engine/Plugins/Online/Experimental/OnlineServices/Source/OnlineServicesCommon/Public/Online/UserFileCommon.h b/Engine/Plugins/Online/Experimental/OnlineServices/Source/OnlineServicesCommon/Public/Online/UserFileCommon.h
new file mode 100644
--- /dev/null
+++ b/Engine/Plugins/Online/Experimental/OnlineServices/Source/OnlineServicesCommon/Public/Online/UserFileCommon.h
@@ -0,0 +1,31 @@
+// Copyright Epic Games, Inc. All Rights Reserved.
+
+#pragma once
+
+#include "Online/UserFile.h"
+#include "Online/OnlineComponent.h"
+
+namespace UE::Online {
+
+class FOnlineServicesCommon;
+
+class ONLINESERVICESCOMMON_API FUserFileCommon : public TOnlineComponent<IUserFile>
+{
+public:
+	using Super = IUserFile;
+
+	FUserFileCommon(FOnlineServicesCommon& InServices);
+
+	// TOnlineComponent
+	virtual void RegisterCommands() override;
+
+	// IUserFile
+	virtual TOnlineAsyncOpHandle<FUserFileEnumerateFiles> EnumerateFiles(FUserFileEnumerateFiles::Params&& Params) override;
+	virtual TOnlineResult<FUserFileGetEnumeratedFiles> GetEnumeratedFiles(FUserFileGetEnumeratedFiles::Params&& Params) override;
+	virtual TOnlineAsyncOpHandle<FUserFileReadFile> ReadFile(FUserFileReadFile::Params&& Params) override;
+	virtual TOnlineAsyncOpHandle<FUserFileWriteFile> WriteFile(FUserFileWriteFile::Params&& Params) override;
+	virtual TOnlineAsyncOpHandle<FUserFileCopyFile> CopyFile(FUserFileCopyFile::Params&& Params) override;
+	virtual TOnlineAsyncOpHandle<FUserFileDeleteFile> DeleteFile(FUserFileDeleteFile::Params&& Params) override;
+};
+
+/* UE::Online */}
\ No newline at end of file
diff --git a/Engine/Plugins/Online/Experimental/OnlineServices/Source/OnlineServicesInterface/Public/Online/OnlineServices.h b/Engine/Plugins/Online/Experimental/OnlineServices/Source/OnlineServicesInterface/Public/Online/OnlineServices.h
--- a/Engine/Plugins/Online/Experimental/OnlineServices/Source/OnlineServicesInterface/Public/Online/OnlineServices.h
+++ b/Engine/Plugins/Online/Experimental/OnlineServices/Source/OnlineServicesInterface/Public/Online/OnlineServices.h
@@ -24,6 +24,7 @@ using IStatsPtr = TSharedPtr<class IStats>;
 using IConnectivityPtr = TSharedPtr<class IConnectivity>;
 using IPrivilegesPtr = TSharedPtr<class IPrivileges>;
 using ITitleFilePtr = TSharedPtr<class ITitleFile>;
+using IUserFilePtr = TSharedPtr<class IUserFile>;
 
 struct FGetResolvedConnectString
 {
@@ -113,6 +114,11 @@ public:
 	 *
 	 */
 	virtual ITitleFilePtr GetTitleFileInterface() = 0;
+	
+	/**
+	 *
+	 */
+	virtual IUserFilePtr GetUserFileInterface() = 0;
 
 	/** 
 	 * Get the connectivity string used for client travel
diff --git a/Engine/Plugins/Online/Experimental/OnlineServices/Source/OnlineServicesInterface/Public/Online/UserFile.h b/Engine/Plugins/Online/Experimental/OnlineServices/Source/OnlineServicesInterface/Public/Online/UserFile.h
new file mode 100644
--- /dev/null
+++ b/Engine/Plugins/Online/Experimental/OnlineServices/Source/OnlineServicesInterface/Public/Online/UserFile.h
@@ -0,0 +1,208 @@
+// Copyright Epic Games, Inc. All Rights Reserved.
+
+#pragma once
+
+#include "Online/CoreOnline.h"
+#include "Online/OnlineAsyncOpHandle.h"
+#include "Online/OnlineMeta.h"
+
+namespace UE::Online {
+
+struct FUserFileEnumerateFiles
+{
+	static constexpr TCHAR Name[] = TEXT("EnumerateFiles");
+
+	struct Params
+	{
+		/** Local user performing the operation */
+		FOnlineAccountIdHandle LocalUserId;
+	};
+
+	struct Result
+	{
+	};
+};
+
+struct FUserFileGetEnumeratedFiles
+{
+	static constexpr TCHAR Name[] = TEXT("GetEnumeratedFiles");
+
+	struct Params
+	{
+		/** Local user performing the operation */
+		FOnlineAccountIdHandle LocalUserId;
+	};
+
+	struct Result
+	{
+		/** Available files for the local user */
+		TArray<FString> Filenames;
+	};
+};
+
+using FUserFileContents = TArray<uint8>;
+using FUserFileContentsPtr = TSharedPtr<const FUserFileContents>;
+
+struct FUserFileReadFile
+{
+	static constexpr TCHAR Name[] = TEXT("ReadFile");
+
+	struct Params
+	{
+		/** Local user performing the operation */
+		FOnlineAccountIdHandle LocalUserId;
+		/** File to be read */
+		FString Filename;
+	};
+
+	struct Result
+	{
+		/** Contents of the file */
+		FUserFileContentsPtr FileContents;
+	};
+};
+
+struct FUserFileWriteFile
+{
+	static constexpr TCHAR Name[] = TEXT("WriteFile");
+
+	struct Params
+	{
+		/** Local user performing the operation */
+		FOnlineAccountIdHandle LocalUserId;
+		/** File to be written */
+		FString Filename;
+		/** Data to write */
+		FUserFileContents FileContents;
+	};
+
+	struct Result
+	{
+	};
+};
+
+struct FUserFileCopyFile
+{
+	static constexpr TCHAR Name[] = TEXT("CopyFile");
+
+	struct Params
+	{
+		/** Local user performing the operation */
+		FOnlineAccountIdHandle LocalUserId;
+		/** File to copy contents from */
+		FString SourceFilename;
+		/** File to copy contents to */
+		FString TargetFilename;
+	};
+
+	struct Result
+	{
+	};
+};
+
+struct FUserFileDeleteFile
+{
+	static constexpr TCHAR Name[] = TEXT("DeleteFile");
+
+	struct Params
+	{
+		/** Local user performing the operation */
+		FOnlineAccountIdHandle LocalUserId;
+		/** File to delete */
+		FString Filename;
+	};
+
+	struct Result
+	{
+	};
+};
+
+class IUserFile
+{
+public:
+	/**
+	 * Enumerate all available files
+	 */
+	virtual TOnlineAsyncOpHandle<FUserFileEnumerateFiles> EnumerateFiles(FUserFileEnumerateFiles::Params&& Params) = 0;
+
+	/**
+	 * Get the cached list of files enumerated by previous call to EnumerateFiles
+	 */
+	virtual TOnlineResult<FUserFileGetEnumeratedFiles> GetEnumeratedFiles(FUserFileGetEnumeratedFiles::Params&& Params) = 0;
+
+	/**
+	 * Read a file and return the file contents
+	 */
+	virtual TOnlineAsyncOpHandle<FUserFileReadFile> ReadFile(FUserFileReadFile::Params&& Params) = 0;
+
+	/**
+	 * Write file contents to a file.
+	 */
+	virtual TOnlineAsyncOpHandle<FUserFileWriteFile> WriteFile(FUserFileWriteFile::Params&& Params) = 0;
+
+	/**
+	 * Copy file contents to another file
+	 */
+	virtual TOnlineAsyncOpHandle<FUserFileCopyFile> CopyFile(FUserFileCopyFile::Params&& Params) = 0;
+
+	/**
+	 * Delete a file
+	 */
+	virtual TOnlineAsyncOpHandle<FUserFileDeleteFile> DeleteFile(FUserFileDeleteFile::Params&& Params) = 0;
+};
+
+namespace Meta {
+
+BEGIN_ONLINE_STRUCT_META(FUserFileEnumerateFiles::Params)
+	ONLINE_STRUCT_FIELD(FUserFileEnumerateFiles::Params, LocalUserId)
+END_ONLINE_STRUCT_META()
+
+BEGIN_ONLINE_STRUCT_META(FUserFileEnumerateFiles::Result)
+END_ONLINE_STRUCT_META()
+
+BEGIN_ONLINE_STRUCT_META(FUserFileGetEnumeratedFiles::Params)
+	ONLINE_STRUCT_FIELD(FUserFileGetEnumeratedFiles::Params, LocalUserId)
+END_ONLINE_STRUCT_META()
+
+BEGIN_ONLINE_STRUCT_META(FUserFileGetEnumeratedFiles::Result)
+	ONLINE_STRUCT_FIELD(FUserFileGetEnumeratedFiles::Result, Filenames)
+END_ONLINE_STRUCT_META()
+
+BEGIN_ONLINE_STRUCT_META(FUserFileReadFile::Params)
+	ONLINE_STRUCT_FIELD(FUserFileReadFile::Params, LocalUserId),
+	ONLINE_STRUCT_FIELD(FUserFileReadFile::Params, Filename)
+END_ONLINE_STRUCT_META()
+
+BEGIN_ONLINE_STRUCT_META(FUserFileReadFile::Result)
+	ONLINE_STRUCT_FIELD(FUserFileReadFile::Result, FileContents)
+END_ONLINE_STRUCT_META()
+
+BEGIN_ONLINE_STRUCT_META(FUserFileWriteFile::Params)
+	ONLINE_STRUCT_FIELD(FUserFileWriteFile::Params, LocalUserId),
+	ONLINE_STRUCT_FIELD(FUserFileWriteFile::Params, Filename),
+	ONLINE_STRUCT_FIELD(FUserFileWriteFile::Params, FileContents)
+END_ONLINE_STRUCT_META()
+
+BEGIN_ONLINE_STRUCT_META(FUserFileWriteFile::Result)
+END_ONLINE_STRUCT_META()
+
+BEGIN_ONLINE_STRUCT_META(FUserFileCopyFile::Params)
+	ONLINE_STRUCT_FIELD(FUserFileCopyFile::Params, LocalUserId),
+	ONLINE_STRUCT_FIELD(FUserFileCopyFile::Params, SourceFilename),
+	ONLINE_STRUCT_FIELD(FUserFileCopyFile::Params, TargetFilename)
+END_ONLINE_STRUCT_META()
+
+BEGIN_ONLINE_STRUCT_META(FUserFileCopyFile::Result)
+END_ONLINE_STRUCT_META()
+
+BEGIN_ONLINE_STRUCT_META(FUserFileDeleteFile::Params)
+	ONLINE_STRUCT_FIELD(FUserFileDeleteFile::Params, LocalUserId),
+	ONLINE_STRUCT_FIELD(FUserFileDeleteFile::Params, Filename)
+END_ONLINE_STRUCT_META()
+
+BEGIN_ONLINE_STRUCT_META(FUserFileDeleteFile::Result)
+END_ONLINE_STRUCT_META()
+
+/* Meta*/ }
+
+/* UE::Online */ }
