- Add operator= for FAchievementStateUpdated

diff --git a/Engine/Plugins/Online/Experimental/OnlineServices/Source/OnlineServicesInterface/Private/Online/Achievements.cpp b/Engine/Plugins/Online/Experimental/OnlineServices/Source/OnlineServicesInterface/Private/Online/Achievements.cpp
--- a/Engine/Plugins/Online/Experimental/OnlineServices/Source/OnlineServicesInterface/Private/Online/Achievements.cpp
+++ b/Engine/Plugins/Online/Experimental/OnlineServices/Source/OnlineServicesInterface/Private/Online/Achievements.cpp
@@ -44,4 +44,11 @@ FString ToLogString(const FAchievementState& State)
 		*State.AchievementId, State.Progress, *State.UnlockTime.ToString());
 }
 
+bool operator==(const FAchievementStateUpdated& A, const FAchievementStateUpdated& B)
+{
+	return A.LocalUserId == B.LocalUserId
+		&& A.AchievementIds == B.AchievementIds;
+}
+
+
 /* UE::Online */ }
diff --git a/Engine/Plugins/Online/Experimental/OnlineServices/Source/OnlineServicesInterface/Public/Online/Achievements.h b/Engine/Plugins/Online/Experimental/OnlineServices/Source/OnlineServicesInterface/Public/Online/Achievements.h
--- a/Engine/Plugins/Online/Experimental/OnlineServices/Source/OnlineServicesInterface/Public/Online/Achievements.h
+++ b/Engine/Plugins/Online/Experimental/OnlineServices/Source/OnlineServicesInterface/Public/Online/Achievements.h
@@ -185,6 +185,8 @@ struct FAchievementStateUpdated
 	TArray<FString> AchievementIds;
 };
 
+bool ONLINESERVICESINTERFACE_API operator==(const FAchievementStateUpdated&, const FAchievementStateUpdated&);
+
 class IAchievements
 {
 public:
