- Change Result field from Ptr to Ref

diff --git a/Engine/Plugins/Online/Experimental/OnlineServices/Source/OnlineServicesInterface/Public/Online/UserFile.h b/Engine/Plugins/Online/Experimental/OnlineServices/Source/OnlineServicesInterface/Public/Online/UserFile.h
--- a/Engine/Plugins/Online/Experimental/OnlineServices/Source/OnlineServicesInterface/Public/Online/UserFile.h
+++ b/Engine/Plugins/Online/Experimental/OnlineServices/Source/OnlineServicesInterface/Public/Online/UserFile.h
@@ -41,7 +41,7 @@ struct FUserFileGetEnumeratedFiles
 };
 
 using FUserFileContents = TArray<uint8>;
-using FUserFileContentsPtr = TSharedPtr<const FUserFileContents>;
+using FUserFileContentsRef = TSharedRef<const FUserFileContents>;
 
 struct FUserFileReadFile
 {
@@ -58,7 +58,7 @@ struct FUserFileReadFile
 	struct Result
 	{
 		/** Contents of the file */
-		FUserFileContentsPtr FileContents;
+		FUserFileContentsRef FileContents;
 	};
 };
 
