OSSv2: Add Stats null implementation.
Add stat definitions in common implementation and corresponding .ini config loading code.

diff --git a/Engine/Plugins/Online/Experimental/OnlineServices/Source/OnlineServicesCommon/Private/Online/StatsCommon.cpp b/Engine/Plugins/Online/Experimental/OnlineServices/Source/OnlineServicesCommon/Private/Online/StatsCommon.cpp
--- a/Engine/Plugins/Online/Experimental/OnlineServices/Source/OnlineServicesCommon/Private/Online/StatsCommon.cpp
+++ b/Engine/Plugins/Online/Experimental/OnlineServices/Source/OnlineServicesCommon/Private/Online/StatsCommon.cpp
@@ -7,6 +7,140 @@
 
 namespace UE::Online {
 
+const TCHAR* LexToString(EStatModifyMethod Value)
+{
+	switch (Value)
+	{
+	case EStatModifyMethod::Sum:
+		return TEXT("Sum");
+	case EStatModifyMethod::Largest:
+		return TEXT("Largest");
+	case EStatModifyMethod::Smallest:
+		return TEXT("Smallest");
+	default: checkNoEntry(); // Intentional fallthrough
+	case EStatModifyMethod::Set:
+		return TEXT("Set");
+	}
+}
+
+void LexFromString(EStatModifyMethod& OutValue, const TCHAR* InStr)
+{
+	if (FCString::Stricmp(InStr, TEXT("Sum")) == 0)
+	{
+		OutValue = EStatModifyMethod::Sum;
+	}
+	else if (FCString::Stricmp(InStr, TEXT("Set")) == 0)
+	{
+		OutValue = EStatModifyMethod::Set;
+	}
+	else if (FCString::Stricmp(InStr, TEXT("Largest")) == 0)
+	{
+		OutValue = EStatModifyMethod::Largest;
+	}
+	else if (FCString::Stricmp(InStr, TEXT("Smallest")) == 0)
+	{
+		OutValue = EStatModifyMethod::Smallest;
+	}
+	else
+	{
+		ensureMsgf(false, TEXT("Can't convert %s to EStatModifyMethod"), InStr);
+		OutValue = EStatModifyMethod::Set;
+	}
+}
+
+const TCHAR* LexToString(EStatLeaderboardUpdateMethod Value)
+{
+	switch (Value)
+	{
+	case EStatLeaderboardUpdateMethod::Force:
+		return TEXT("Force");
+	default: checkNoEntry(); // Intentional fallthrough
+	case EStatLeaderboardUpdateMethod::KeepBest:
+		return TEXT("KeepBest");
+	}
+}
+
+void LexFromString(EStatLeaderboardUpdateMethod& OutValue, const TCHAR* InStr)
+{
+	if (FCString::Stricmp(InStr, TEXT("KeepBest")) == 0)
+	{
+		OutValue = EStatLeaderboardUpdateMethod::KeepBest;
+	}
+	else if (FCString::Stricmp(InStr, TEXT("Set")) == 0)
+	{
+		OutValue = EStatLeaderboardUpdateMethod::Force;
+	}
+	else
+	{
+		ensureMsgf(false, TEXT("Can't convert %s to EStatLeaderboardUpdateMethod"), InStr);
+		OutValue = EStatLeaderboardUpdateMethod::KeepBest;
+	}
+}
+
+const TCHAR* LexToString(EStatLeaderboardOrderMethod Value)
+{
+	switch (Value)
+	{
+	case EStatLeaderboardOrderMethod::Ascending:
+		return TEXT("Ascending");
+	default: checkNoEntry(); // Intentional fallthrough
+	case EStatLeaderboardOrderMethod::Descending:
+		return TEXT("Descending");
+	}
+}
+
+void LexFromString(EStatLeaderboardOrderMethod& OutValue, const TCHAR* InStr)
+{
+	if (FCString::Stricmp(InStr, TEXT("Ascending")) == 0)
+	{
+		OutValue = EStatLeaderboardOrderMethod::Ascending;
+	}
+	else if (FCString::Stricmp(InStr, TEXT("Descending")) == 0)
+	{
+		OutValue = EStatLeaderboardOrderMethod::Descending;
+	}
+	else
+	{
+		ensureMsgf(false, TEXT("Can't convert %s to EStatLeaderboardOrderMethod"), InStr);
+		OutValue = EStatLeaderboardOrderMethod::Descending;
+	}
+}
+
+const TCHAR* LexToString(EStatUsageFlags Value)
+{
+	switch (Value)
+	{
+	case EStatUsageFlags::Achievement:
+		return TEXT("Achievement");
+	case EStatUsageFlags::Leaderboard:
+		return TEXT("Leaderboard");
+	default: checkNoEntry(); // Intentional fallthrough
+	case EStatUsageFlags::None:
+		return TEXT("None");
+	}
+}
+
+void LexFromString(EStatUsageFlags& OutValue, const TCHAR* InStr)
+{
+	if (FCString::Stricmp(InStr, TEXT("None")) == 0)
+	{
+		OutValue = EStatUsageFlags::None;
+	}
+	else if (FCString::Stricmp(InStr, TEXT("Achievement")) == 0)
+	{
+		OutValue = EStatUsageFlags::Achievement;
+	}
+	else if (FCString::Stricmp(InStr, TEXT("Leaderboard")) == 0)
+	{
+		OutValue = EStatUsageFlags::Leaderboard;
+	}
+	else
+	{
+		ensureMsgf(false, TEXT("Can't convert %s to EStatusageFlags"), InStr);
+		OutValue = EStatUsageFlags::None;
+	}
+}
+
 FStatsCommon::FStatsCommon(FOnlineServicesCommon& InServices)
 	: TOnlineComponent(TEXT("Stats"), InServices)
 {
@@ -14,9 +148,70 @@ FStatsCommon::FStatsCommon(FOnlineServicesCommon& InServices)
 
 void FStatsCommon::Initialize()
 {
+	ReadStatDefinitionsFromConfig();
+
 	TOnlineComponent<IStats>::Initialize();
 }
 
+void FStatsCommon::ReadStatDefinitionsFromConfig()
+{
+	const TCHAR* ConfigSection = TEXT("OnlineServices.Stats");
+
+	for (int StatIdx = 0;; StatIdx++)
+	{
+		FString StatName;
+		GConfig->GetString(ConfigSection, *FString::Printf(TEXT("StatDef_%d_Name"), StatIdx), StatName, GEngineIni);
+		if (StatName.IsEmpty())
+		{
+			break;
+		}
+
+		FStatDefinition& StatDefinition = StatDefinitions.Emplace(StatName);
+		StatDefinition.Name = MoveTemp(StatName);
+
+		GConfig->GetInt(ConfigSection, *FString::Printf(TEXT("StatDef_%d_Id"), StatIdx), StatDefinition.Id, GEngineIni);
+
+		FText StatUsageFlagsStr;
+		GConfig->GetText(ConfigSection, *FString::Printf(TEXT("StatDef_%d_UsageFlags"), StatIdx), StatUsageFlagsStr, GEngineIni);
+		TArray<FString> StatUsageFlagStrArray;
+		bool CullEmpty = true;
+		StatUsageFlagsStr.ToString().ParseIntoArray(StatUsageFlagStrArray, TEXT(","), CullEmpty);
+		for (const FString& StatUsageFlagStr : StatUsageFlagStrArray)
+		{
+			EStatUsageFlags StatUsageFlag = EStatUsageFlags::None;
+			LexFromString(StatUsageFlag, *StatUsageFlagStr);
+			StatDefinition.UsageFlags |= (uint32)StatUsageFlag;
+		}
+
+		if (StatDefinition.UsageFlags & (uint32)EStatUsageFlags::Achievement)
+		{
+			FText StatModifyMethod;
+			GConfig->GetText(ConfigSection, *FString::Printf(TEXT("StatDef_%d_ModifyMethod"), StatIdx), StatModifyMethod, GEngineIni);
+			if (!StatModifyMethod.IsEmpty())
+			{
+				LexFromString(StatDefinition.ModifyMethod, *StatModifyMethod.ToString());
+			}
+		}
+
+		if (StatDefinition.UsageFlags & (uint32)EStatUsageFlags::Leaderboard)
+		{
+			FText StatLeaderboardUpdateMethod;
+			GConfig->GetText(ConfigSection, *FString::Printf(TEXT("StatDef_%d_LeaderboardUpdateMethod"), StatIdx), StatLeaderboardUpdateMethod, GEngineIni);
+			if (!StatLeaderboardUpdateMethod.IsEmpty())
+			{
+				LexFromString(StatDefinition.LeaderboardUpdateMethod, *StatLeaderboardUpdateMethod.ToString());
+			}
+
+			FText StatLeaderboardOrderMethod;
+			GConfig->GetText(ConfigSection, *FString::Printf(TEXT("StatDef_%d_LeaderboardOrderMethod"), StatIdx), StatLeaderboardOrderMethod, GEngineIni);
+			if (!StatLeaderboardOrderMethod.IsEmpty())
+			{
+				LexFromString(StatDefinition.LeaderboardOrderMethod, *StatLeaderboardOrderMethod.ToString());
+			}
+		}
+	}
+}
+
 void FStatsCommon::RegisterCommands()
 {
 	TOnlineComponent<IStats>::RegisterCommands();
diff --git a/Engine/Plugins/Online/Experimental/OnlineServices/Source/OnlineServicesCommon/Public/Online/StatsCommon.h b/Engine/Plugins/Online/Experimental/OnlineServices/Source/OnlineServicesCommon/Public/Online/StatsCommon.h
--- a/Engine/Plugins/Online/Experimental/OnlineServices/Source/OnlineServicesCommon/Public/Online/StatsCommon.h
+++ b/Engine/Plugins/Online/Experimental/OnlineServices/Source/OnlineServicesCommon/Public/Online/StatsCommon.h
@@ -7,24 +7,96 @@
 
 namespace UE::Online {
 
-	class ONLINESERVICESCOMMON_API FStatsCommon : public TOnlineComponent<IStats>
-	{
-	public:
-		using Super = IStats;
+enum class EStatModifyMethod : uint8
+{
+	/** Add the new value to the previous value */
+	Sum,
+	/** Overwrite previous value with the new value */
+	Set,
+	/** Only replace previous value if new value is larger */
+	Largest,
+	/** Only replace previous value if new value is smaller */
+	Smallest
+};
 
-		FStatsCommon(FOnlineServicesCommon& InServices);
+const TCHAR* LexToString(EStatModifyMethod Value);
+void LexFromString(EStatModifyMethod& OutValue, const TCHAR* InStr);
 
-		// TOnlineComponent
-		virtual void Initialize() override;
-		virtual void RegisterCommands() override;
+enum class EStatLeaderboardUpdateMethod : uint8
+{
+	/** If current leaderboard score is better than the uploaded one, keep the current one */
+	KeepBest,
+	/** Leaderboard score is always replaced with uploaded value */
+	Force
+};
 
-		// IStats
-		virtual TOnlineAsyncOpHandle<FUpdateStats> UpdateStats(FUpdateStats::Params&& Params) override;
-		virtual TOnlineAsyncOpHandle<FQueryStats> QueryStats(FQueryStats::Params&& Params) override;
-		virtual TOnlineAsyncOpHandle<FBatchQueryStats> BatchQueryStats(FBatchQueryStats::Params&& Params) override;
+const TCHAR* LexToString(EStatLeaderboardUpdateMethod Value);
+void LexFromString(EStatLeaderboardUpdateMethod& OutValue, const TCHAR* InStr);
+
+enum EStatLeaderboardOrderMethod : uint8
+{
+	Ascending,
+	Descending
+};
+
+const TCHAR* LexToString(EStatLeaderboardOrderMethod Value);
+void LexFromString(EStatLeaderboardOrderMethod& OutValue, const TCHAR* InStr);
+
+enum class EStatUsageFlags : uint8
+{
+	None = 0,
+	Achievement = 1 << 0,
+	Leaderboard = 1 << 1
+};
+
+ENUM_CLASS_FLAGS(EStatUsageFlags);
+
+const TCHAR* LexToString(EStatUsageFlags Value);
+void LexFromString(EStatUsageFlags& OutValue, const TCHAR* InStr);
+
+struct FStatDefinition
+{
+	/* The name of the stat */
+	FString Name;
+	/* Corresponding stat id on the platform if needed */
+	int32 Id = 0;
+
+	/* What is this stat used for, in array format split by ",", for example: "Achievement,Leaderboard" */
+	int32 UsageFlags = 0;
+
+	/* How the stat will be modified, only useful when EStatUsageFlags::Achievement is set in UsageFlags */
+	EStatModifyMethod ModifyMethod = EStatModifyMethod::Set;
+
+	/* How the leaderboard score will be updated, only useful when EStatUsageFlags::Leaderboard is set in UsageFlags */
+	EStatLeaderboardUpdateMethod LeaderboardUpdateMethod = EStatLeaderboardUpdateMethod::Force;
+	/* How the leaderboard score will be ordered, only useful when EStatUsageFlags::Leaderboard is set in UsageFlags */
+	EStatLeaderboardOrderMethod LeaderboardOrderMethod = EStatLeaderboardOrderMethod::Descending;
+};
+
+
+class ONLINESERVICESCOMMON_API FStatsCommon : public TOnlineComponent<IStats>
+{
+public:
+	using Super = IStats;
+
+	FStatsCommon(FOnlineServicesCommon& InServices);
+
+	// TOnlineComponent
+	virtual void Initialize() override;
+	virtual void RegisterCommands() override;
+
+	// IStats
+	virtual TOnlineAsyncOpHandle<FUpdateStats> UpdateStats(FUpdateStats::Params&& Params) override;
+	virtual TOnlineAsyncOpHandle<FQueryStats> QueryStats(FQueryStats::Params&& Params) override;
+	virtual TOnlineAsyncOpHandle<FBatchQueryStats> BatchQueryStats(FBatchQueryStats::Params&& Params) override;
 #if !UE_BUILD_SHIPPING
-		virtual TOnlineAsyncOpHandle<FResetStats> ResetStats(FResetStats::Params&& Params) override;
+	virtual TOnlineAsyncOpHandle<FResetStats> ResetStats(FResetStats::Params&& Params) override;
 #endif // !UE_BUILD_SHIPPING
-	};
+
+protected:
+	void ReadStatDefinitionsFromConfig();
+
+	TMap<FString, FStatDefinition> StatDefinitions;
+};
 
 /* UE::Online */ }
diff --git a/Engine/Plugins/Online/Experimental/OnlineServices/Source/OnlineServicesInterface/Public/Online/Stats.h b/Engine/Plugins/Online/Experimental/OnlineServices/Source/OnlineServicesInterface/Public/Online/Stats.h
--- a/Engine/Plugins/Online/Experimental/OnlineServices/Source/OnlineServicesInterface/Public/Online/Stats.h
+++ b/Engine/Plugins/Online/Experimental/OnlineServices/Source/OnlineServicesInterface/Public/Online/Stats.h
@@ -8,8 +8,6 @@
 
 namespace UE::Online {
 
-class FOnlineError;
-
 //using FStatValue = SchemaVariant;
 using FStatValue = TVariant<bool, int32, int64, float, double, FString>;
 
@@ -71,7 +69,7 @@ struct FBatchQueryStats
 
 	struct Result
 	{
-		/* The stats of queried users */
+		/* The stats of queried users, there should be no much elements so simply use TArray instead of TMap */
 		TArray<FUserStats> UsersStats;
 	};
 };
diff --git a/Engine/Plugins/Online/Experimental/OnlineServicesNull/Source/Private/Online/OnlineServicesNull.cpp b/Engine/Plugins/Online/Experimental/OnlineServicesNull/Source/Private/Online/OnlineServicesNull.cpp
--- a/Engine/Plugins/Online/Experimental/OnlineServicesNull/Source/Private/Online/OnlineServicesNull.cpp
+++ b/Engine/Plugins/Online/Experimental/OnlineServicesNull/Source/Private/Online/OnlineServicesNull.cpp
@@ -5,6 +5,7 @@
 
 #include "Online/AchievementsNull.h"
 #include "Online/AuthNull.h"
+#include "Online/StatsNull.h"
 #include "Online/LobbiesNull.h"
 #include "Online/TitleFileNull.h"
 
@@ -34,6 +35,7 @@ void FOnlineServicesNull::RegisterComponents()
 {
 	Components.Register<FAchievementsNull>(*this);
 	Components.Register<FAuthNull>(*this);
+	Components.Register<FStatsNull>(*this);
 #if WITH_ENGINE
 	Components.Register<FLobbiesNull>(*this);
 #endif
diff --git a/Engine/Plugins/Online/Experimental/OnlineServicesNull/Source/Private/Online/StatsNull.cpp b/Engine/Plugins/Online/Experimental/OnlineServicesNull/Source/Private/Online/StatsNull.cpp
new file mode 100644
--- /dev/null
+++ b/Engine/Plugins/Online/Experimental/OnlineServicesNull/Source/Private/Online/StatsNull.cpp
@@ -0,0 +1,175 @@
+// Copyright Epic Games, Inc. All Rights Reserved.
+
+#include "Online/StatsNull.h"
+#include "Online/AuthNull.h"
+#include "Online/OnlineServicesNull.h"
+
+namespace UE::Online {
+
+FStatsNull::FStatsNull(FOnlineServicesNull& InOwningSubsystem)
+	: Super(InOwningSubsystem)
+{
+}
+
+struct FFindUserStatsNull
+{
+	FFindUserStatsNull(const FOnlineAccountIdHandle& InUserId)
+		: UserId(InUserId)
+	{
+	}
+
+	bool operator()(const FUserStats& UserStats)
+	{
+		return UserStats.UserId == UserId;
+	}
+
+	FOnlineAccountIdHandle UserId;
+};
+
+TOnlineAsyncOpHandle<FUpdateStats> FStatsNull::UpdateStats(FUpdateStats::Params&& Params)
+{
+	TOnlineAsyncOpRef<FUpdateStats> Op = GetOp<FUpdateStats>(MoveTemp(Params));
+
+	if (!Services.Get<FAuthNull>()->IsLoggedIn(Op->GetParams().LocalUserId))
+	{
+		Op->SetError(Errors::InvalidUser());
+		return Op->GetHandle();
+	}
+
+	Op->Then([this](TOnlineAsyncOp<FUpdateStats>& InAsyncOp) mutable
+	{
+		for (const FUserStats& UpdateUserStats : InAsyncOp.GetParams().UpdateUsersStats)
+		{
+			FUserStats* ExistingUserStats = UsersStats.FindByPredicate(FFindUserStatsNull(UpdateUserStats.UserId));
+			if (!ExistingUserStats)
+			{
+				ExistingUserStats = &UsersStats.Emplace_GetRef();
+				ExistingUserStats->UserId = UpdateUserStats.UserId;
+			}
+
+			TMap<FString, FStatValue>& UserStats = ExistingUserStats->Stats;
+			for (const auto& UpdateUserStatPair : UpdateUserStats.Stats)
+			{
+				if (FStatDefinition* StatDefinition = StatDefinitions.Find(UpdateUserStatPair.Key))
+				{
+					if (StatDefinition->UsageFlags & (uint32)EStatUsageFlags::Achievement)
+					{
+						if (FStatValue* StatValue = UserStats.Find(UpdateUserStatPair.Key))
+						{
+							*StatValue = UpdateUserStatPair.Value;
+
+							// TODO: Apply the value properly when SchemaVariant is ready
+							//switch (StatDefinition->ModifyMethod)
+							//{
+							//case EStatModifyMethod::Set:
+							//	*StatValue = UpdateUserStatPair.Value;
+							//	break;
+							//case EStatModifyMethod::Sum:
+							//	*StatValue += UpdateUserStatPair.Value;
+							//	break;
+							//case EStatModifyMethod::Largest:
+							//	*StatValue = Max(*StatValue, UpdateUserStatPair.Value);
+							//	break;
+							//case EStatModifyMethod::Smallest:
+							//	*StatValue = Min(*StatValue, UpdateUserStatPair.Value);
+							//	break;
+							//}
+						}
+						else
+						{
+							UserStats[UpdateUserStatPair.Key] = UpdateUserStatPair.Value;
+						}
+					}
+
+					if (StatDefinition->UsageFlags & (uint32)EStatUsageFlags::Leaderboard)
+					{
+						// TODO: Call LeaderboardsNull to update
+					}
+				}
+			}
+		}
+
+		InAsyncOp.SetResult({});
+	})
+	.Enqueue(GetSerialQueue());
+
+	return Op->GetHandle();
+}
+
+TOnlineAsyncOpHandle<FQueryStats> FStatsNull::QueryStats(FQueryStats::Params&& Params)
+{
+	TOnlineAsyncOpRef<FQueryStats> Op = GetOp<FQueryStats>(MoveTemp(Params));
+
+	if (!Services.Get<FAuthNull>()->IsLoggedIn(Op->GetParams().LocalUserId))
+	{
+		Op->SetError(Errors::InvalidUser());
+		return Op->GetHandle();
+	}
+
+	Op->Then([this](TOnlineAsyncOp<FQueryStats>& InAsyncOp)
+	{
+		FQueryStats::Result Result;
+
+		if (FUserStats* ExistingUserStats = UsersStats.FindByPredicate(FFindUserStatsNull(InAsyncOp.GetParams().TargetUserId)))
+		{
+			Result.Stats = ExistingUserStats->Stats;
+		}
+
+		InAsyncOp.SetResult(MoveTemp(Result));
+	});
+	return Op->GetHandle();
+}
+
+TOnlineAsyncOpHandle<FBatchQueryStats> FStatsNull::BatchQueryStats(FBatchQueryStats::Params&& Params)
+{
+	TOnlineAsyncOpRef<FBatchQueryStats> Op = GetOp<FBatchQueryStats>(MoveTemp(Params));
+	if (!Services.Get<FAuthNull>()->IsLoggedIn(Op->GetParams().LocalUserId))
+	{
+		Op->SetError(Errors::InvalidUser());
+		return Op->GetHandle();
+	}
+
+	Op->Then([this](TOnlineAsyncOp<FBatchQueryStats>& InAsyncOp)
+	{
+		FBatchQueryStats::Result Result;
+
+		for (const FOnlineAccountIdHandle& TargetUserId : InAsyncOp.GetParams().TargetUserIds)
+		{
+			if (FUserStats* ExistingUserStats = UsersStats.FindByPredicate(FFindUserStatsNull(TargetUserId)))
+			{
+				Result.UsersStats.Emplace(*ExistingUserStats);
+			}
+		}
+
+		InAsyncOp.SetResult(MoveTemp(Result));
+	});
+
+	return Op->GetHandle();
+}
+
+#if !UE_BUILD_SHIPPING
+TOnlineAsyncOpHandle<FResetStats> FStatsNull::ResetStats(FResetStats::Params&& Params)
+{
+	TOnlineAsyncOpRef<FResetStats> Op = GetOp<FResetStats>(MoveTemp(Params));
+
+	if (!Services.Get<FAuthNull>()->IsLoggedIn(Op->GetParams().LocalUserId))
+	{
+		Op->SetError(Errors::InvalidUser());
+		return Op->GetHandle();
+	}
+
+	Op->Then([this](TOnlineAsyncOp<FResetStats>& InAsyncOp)
+	{
+		uint32 Index = UsersStats.IndexOfByPredicate(FFindUserStatsNull(InAsyncOp.GetParams().LocalUserId));
+		if (Index != INDEX_NONE)
+		{
+			UsersStats.RemoveAt(Index);
+		}
+		InAsyncOp.SetResult({});
+	});
+
+	return Op->GetHandle();
+}
+#endif // !UE_BUILD_SHIPPING
+
+/* UE::Online */ }
diff --git a/Engine/Plugins/Online/Experimental/OnlineServicesNull/Source/Public/Online/StatsNull.h b/Engine/Plugins/Online/Experimental/OnlineServicesNull/Source/Public/Online/StatsNull.h
new file mode 100644
--- /dev/null
+++ b/Engine/Plugins/Online/Experimental/OnlineServicesNull/Source/Public/Online/StatsNull.h
@@ -0,0 +1,31 @@
+// Copyright Epic Games, Inc. All Rights Reserved.
+
+#pragma once
+
+#include "Misc/Optional.h"
+#include "Online/StatsCommon.h"
+
+namespace UE::Online {
+
+class FOnlineServicesNull;
+
+class ONLINESERVICESNULL_API FStatsNull : public FStatsCommon
+{
+public:
+	using Super = FStatsCommon;
+
+	FStatsNull(FOnlineServicesNull& InOwningSubsystem);
+
+	// IStats
+	virtual TOnlineAsyncOpHandle<FUpdateStats> UpdateStats(FUpdateStats::Params&& Params) override;
+	virtual TOnlineAsyncOpHandle<FQueryStats> QueryStats(FQueryStats::Params&& Params) override;
+	virtual TOnlineAsyncOpHandle<FBatchQueryStats> BatchQueryStats(FBatchQueryStats::Params&& Params) override;
+#if !UE_BUILD_SHIPPING
+	virtual TOnlineAsyncOpHandle<FResetStats> ResetStats(FResetStats::Params&& Params) override;
+#endif // !UE_BUILD_SHIPPING
+
+protected:
+	TArray<FUserStats> UsersStats;
+};
+
+/* UE::Online */ }
