Mark TOnlineEvent::Add as no discard to require callers to capture the result.
Without this, their event binding does not stay bound.

diff --git a/Engine/Plugins/Online/Experimental/OnlineServices/Source/OnlineServicesInterface/Public/Online/OnlineAsyncOpHandle.h b/Engine/Plugins/Online/Experimental/OnlineServices/Source/OnlineServicesInterface/Public/Online/OnlineAsyncOpHandle.h
--- a/Engine/Plugins/Online/Experimental/OnlineServices/Source/OnlineServicesInterface/Public/Online/OnlineAsyncOpHandle.h
+++ b/Engine/Plugins/Online/Experimental/OnlineServices/Source/OnlineServicesInterface/Public/Online/OnlineAsyncOpHandle.h
@@ -347,7 +347,7 @@ public:
 	 *
 	 */
 	template<typename... ConstructDelegateParamTypes>
-	FOnlineEventDelegateHandle Add(ConstructDelegateParamTypes&&... Params)
+	UE_NODISCARD FOnlineEventDelegateHandle Add(ConstructDelegateParamTypes&&... Params)
 	{
 		return Event.Add(Forward<ConstructDelegateParamTypes>(Params)...);
 	}
