OnlineServices - UserInfo Interface and common implementation

diff --git a/Engine/Plugins/Online/Experimental/OnlineServices/Source/OnlineServicesCommon/Private/Online/OnlineServicesCommon.cpp b/Engine/Plugins/Online/Experimental/OnlineServices/Source/OnlineServicesCommon/Private/Online/OnlineServicesCommon.cpp
--- a/Engine/Plugins/Online/Experimental/OnlineServices/Source/OnlineServicesCommon/Private/Online/OnlineServicesCommon.cpp
+++ b/Engine/Plugins/Online/Experimental/OnlineServices/Source/OnlineServicesCommon/Private/Online/OnlineServicesCommon.cpp
@@ -50,6 +50,11 @@ IAuthPtr FOnlineServicesCommon::GetAuthInterface()
 	return IAuthPtr(AsShared(), Get<IAuth>());
 }
 
+IUserInfoPtr FOnlineServicesCommon::GetUserInfoInterface()
+{
+	return IUserInfoPtr(AsShared(), Get<IUserInfo>());
+}
+
 IFriendsPtr FOnlineServicesCommon::GetFriendsInterface()
 {
 	return IFriendsPtr(AsShared(), Get<IFriends>());
diff --git a/Engine/Plugins/Online/Experimental/OnlineServices/Source/OnlineServicesCommon/Private/Online/UserInfoCommon.cpp b/Engine/Plugins/Online/Experimental/OnlineServices/Source/OnlineServicesCommon/Private/Online/UserInfoCommon.cpp
new file mode 100644
--- /dev/null
+++ b/Engine/Plugins/Online/Experimental/OnlineServices/Source/OnlineServicesCommon/Private/Online/UserInfoCommon.cpp
@@ -0,0 +1,54 @@
+// Copyright Epic Games, Inc. All Rights Reserved.
+
+#include "Online/UserInfoCommon.h"
+
+#include "Online/OnlineServicesCommon.h"
+
+namespace UE::Online {
+
+FUserInfoCommon::FUserInfoCommon(FOnlineServicesCommon& InServices)
+	: TOnlineComponent(TEXT("UserInfo"), InServices)
+{
+}
+
+void FUserInfoCommon::RegisterCommands()
+{
+	RegisterCommand(&FUserInfoCommon::QueryUserInfo);
+	RegisterCommand(&FUserInfoCommon::GetUserInfo);
+	RegisterCommand(&FUserInfoCommon::QueryUserAvatar);
+	RegisterCommand(&FUserInfoCommon::GetUserAvatar);
+	RegisterCommand(&FUserInfoCommon::ShowUserProfile);
+}
+
+TOnlineAsyncOpHandle<FQueryUserInfo> FUserInfoCommon::QueryUserInfo(FQueryUserInfo::Params&& Params)
+{
+	TOnlineAsyncOpRef<FQueryUserInfo> Operation = GetOp<FQueryUserInfo>(MoveTemp(Params));
+	Operation->SetError(Errors::NotImplemented());
+	return Operation->GetHandle();
+}
+
+TOnlineResult<FGetUserInfo> FUserInfoCommon::GetUserInfo(FGetUserInfo::Params&& Params)
+{
+	return TOnlineResult<FGetUserInfo>(Errors::NotImplemented());
+}
+
+TOnlineAsyncOpHandle<FQueryUserAvatar> FUserInfoCommon::QueryUserAvatar(FQueryUserAvatar::Params&& Params)
+{
+	TOnlineAsyncOpRef<FQueryUserAvatar> Operation = GetOp<FQueryUserAvatar>(MoveTemp(Params));
+	Operation->SetError(Errors::NotImplemented());
+	return Operation->GetHandle();
+}
+
+TOnlineResult<FGetUserAvatar> FUserInfoCommon::GetUserAvatar(FGetUserAvatar::Params&& Params)
+{
+	return TOnlineResult<FGetUserAvatar>(Errors::NotImplemented());
+}
+
+TOnlineAsyncOpHandle<FShowUserProfile> FUserInfoCommon::ShowUserProfile(FShowUserProfile::Params&& Params)
+{
+	TOnlineAsyncOpRef<FShowUserProfile> Operation = GetOp<FShowUserProfile>(MoveTemp(Params));
+	Operation->SetError(Errors::NotImplemented());
+	return Operation->GetHandle();
+}
+
+/* UE::Online */ }
diff --git a/Engine/Plugins/Online/Experimental/OnlineServices/Source/OnlineServicesCommon/Public/Online/OnlineServicesCommon.h b/Engine/Plugins/Online/Experimental/OnlineServices/Source/OnlineServicesCommon/Public/Online/OnlineServicesCommon.h
--- a/Engine/Plugins/Online/Experimental/OnlineServices/Source/OnlineServicesCommon/Public/Online/OnlineServicesCommon.h
+++ b/Engine/Plugins/Online/Experimental/OnlineServices/Source/OnlineServicesCommon/Public/Online/OnlineServicesCommon.h
@@ -37,6 +37,7 @@ public:
 	virtual void Destroy() override;
 	virtual IAchievementsPtr GetAchievementsInterface() override;
 	virtual IAuthPtr GetAuthInterface() override;
+	virtual IUserInfoPtr GetUserInfoInterface() override;
 	virtual IFriendsPtr GetFriendsInterface() override;
 	virtual IPresencePtr GetPresenceInterface() override;
 	virtual IExternalUIPtr GetExternalUIInterface() override;
diff --git a/Engine/Plugins/Online/Experimental/OnlineServices/Source/OnlineServicesCommon/Public/Online/UserInfoCommon.h b/Engine/Plugins/Online/Experimental/OnlineServices/Source/OnlineServicesCommon/Public/Online/UserInfoCommon.h
new file mode 100644
--- /dev/null
+++ b/Engine/Plugins/Online/Experimental/OnlineServices/Source/OnlineServicesCommon/Public/Online/UserInfoCommon.h
@@ -0,0 +1,30 @@
+// Copyright Epic Games, Inc. All Rights Reserved.
+
+#pragma once
+
+#include "Online/UserInfo.h"
+#include "Online/OnlineComponent.h"
+
+namespace UE::Online {
+
+class FOnlineServicesCommon;
+
+class ONLINESERVICESCOMMON_API FUserInfoCommon : public TOnlineComponent<IUserInfo>
+{
+public:
+	using Super = IUserInfo;
+
+	FUserInfoCommon(FOnlineServicesCommon& InServices);
+
+	// TOnlineComponent
+	virtual void RegisterCommands() override;
+
+	// IUserInfo
+	virtual TOnlineAsyncOpHandle<FQueryUserInfo> QueryUserInfo(FQueryUserInfo::Params&& Params) override;
+	virtual TOnlineResult<FGetUserInfo> GetUserInfo(FGetUserInfo::Params&& Params) override;
+	virtual TOnlineAsyncOpHandle<FQueryUserAvatar> QueryUserAvatar(FQueryUserAvatar::Params&& Params) override;
+	virtual TOnlineResult<FGetUserAvatar> GetUserAvatar(FGetUserAvatar::Params&& Params) override;
+	virtual TOnlineAsyncOpHandle<FShowUserProfile> ShowUserProfile(FShowUserProfile::Params&& Params) override;
+};
+
+/* UE::Online */ }
diff --git a/Engine/Plugins/Online/Experimental/OnlineServices/Source/OnlineServicesInterface/Public/Online/OnlineServices.h b/Engine/Plugins/Online/Experimental/OnlineServices/Source/OnlineServicesInterface/Public/Online/OnlineServices.h
--- a/Engine/Plugins/Online/Experimental/OnlineServices/Source/OnlineServicesInterface/Public/Online/OnlineServices.h
+++ b/Engine/Plugins/Online/Experimental/OnlineServices/Source/OnlineServicesInterface/Public/Online/OnlineServices.h
@@ -15,6 +15,7 @@ namespace UE::Online {
 using IOnlineServicesPtr = TSharedPtr<class IOnlineServices>;
 using IAchievementsPtr = TSharedPtr<class IAchievements>;
 using IAuthPtr = TSharedPtr<class IAuth>;
+using IUserInfoPtr = TSharedPtr<class IUserInfo>;
 using IFriendsPtr = TSharedPtr<class IFriends>;
 using IPresencePtr = TSharedPtr<class IPresence>;
 using IExternalUIPtr = TSharedPtr<class IExternalUI>;
@@ -67,6 +68,11 @@ public:
 	 */
 	virtual IAuthPtr GetAuthInterface() = 0;
 
+	/**
+	 *
+	 */
+	virtual IUserInfoPtr GetUserInfoInterface() = 0;
+
 	/**
 	 *
 	 */
diff --git a/Engine/Plugins/Online/Experimental/OnlineServices/Source/OnlineServicesInterface/Public/Online/UserInfo.h b/Engine/Plugins/Online/Experimental/OnlineServices/Source/OnlineServicesInterface/Public/Online/UserInfo.h
new file mode 100644
--- /dev/null
+++ b/Engine/Plugins/Online/Experimental/OnlineServices/Source/OnlineServicesInterface/Public/Online/UserInfo.h
@@ -0,0 +1,233 @@
+// Copyright Epic Games, Inc. All Rights Reserved.
+
+#pragma once
+
+#include "Online/CoreOnline.h"
+#include "Online/OnlineAsyncOpHandle.h"
+#include "Online/OnlineMeta.h"
+
+namespace UE::Online {
+
+struct FUserInfo
+{
+	/* Account Id */
+	FOnlineAccountIdHandle UserId;
+	/* Display Name of the User */
+	FString DisplayName;
+};
+
+struct FQueryUserInfo
+{
+	static constexpr TCHAR Name[] = TEXT("QueryUserInfo");
+
+	/** Input struct for UserInfo::QueryUserInfo */
+	struct Params
+	{
+		/* Local user to query users for */
+		FOnlineAccountIdHandle LocalUserId;
+		/* List of User ids to Query*/
+		TArray<FOnlineAccountIdHandle> UserIds;
+	};
+
+	/**
+	 * Output struct for UserInfo::QueryUserInfo
+	 * Obtain the queried user info via GetUserInfo
+	 */
+	struct Result
+	{
+	};
+};
+
+struct FGetUserInfo
+{
+	static constexpr TCHAR Name[] = TEXT("GetUserInfo");
+
+	/** Input struct for UserInfo::GetUserInfo */
+	struct Params
+	{
+		/* Local user */
+		FOnlineAccountIdHandle LocalUserId;
+		/* User id to get info for*/
+		FOnlineAccountIdHandle UserId;
+	};
+
+	/**
+	 * Output struct for UserInfo::GetUserInfo
+	 */
+	struct Result
+	{
+		/* User info */
+		TSharedRef<FUserInfo> UserInfo;
+	};
+};
+
+struct FQueryUserAvatar
+{
+	static constexpr TCHAR Name[] = TEXT("QueryUserAvatar");
+
+	/** Input struct for UserInfo::QueryUserAvatar */
+	struct Params
+	{
+		/* Local user to query users for */
+		FOnlineAccountIdHandle LocalUserId;
+		/* List of User ids to Query*/
+		TArray<FOnlineAccountIdHandle> UserIds;
+	};
+
+	/**
+	 * Output struct for UserInfo::QueryUserAvatar
+	 * Obtain the queried user info via GetUserAvatar
+	 */
+	struct Result
+	{
+	};
+};
+
+struct FGetUserAvatar
+{
+	static constexpr TCHAR Name[] = TEXT("GetUserAvatar");
+
+	/** Input struct for UserInfo::GetUserAvatar */
+	struct Params
+	{
+		/* Local user */
+		FOnlineAccountIdHandle LocalUserId;
+		/* User id to get avatar for*/
+		FOnlineAccountIdHandle UserId;
+	};
+
+	/**
+	 * Output struct for UserInfo::GetUserAvatar
+	 */
+	struct Result
+	{
+		/* Avatar Url */
+		FString AvatarUrl;
+	};
+};
+
+struct FShowUserProfile
+{
+	static constexpr TCHAR Name[] = TEXT("ShowUserProfile");
+
+	/** Input struct for UserInfo::ShowUserProfile */
+	struct Params
+	{
+		/* Local user */
+		FOnlineAccountIdHandle LocalUserId;
+		/* User id to display the profile for*/
+		FOnlineAccountIdHandle UserId;
+	};
+
+	/**
+	 * Output struct for UserInfo::ShowUserProfile
+	 */
+	struct Result
+	{
+	};
+};
+
+class IUserInfo
+{
+public:
+	/*
+	 * Queries the user info for a list of account ids.
+	 * Upon success, user info can be accessed with GetUserInfo.
+	 * @see FUserInfo
+	 * @see GetUserInfo
+	 * 
+	 * @params Params for the QueryUserInfo call
+	 * @return AsyncOpHandle
+	 */
+	virtual TOnlineAsyncOpHandle<FQueryUserInfo> QueryUserInfo(FQueryUserInfo::Params&& Params) = 0;
+
+	/*
+	 * Get the user info for a previously queried AccountId
+	 * @see FUserInfo
+	 * @see QueryUserInfo
+	 *
+	 * @params Params for the GetUserInfo call
+	 * @return Result
+	 */
+	virtual TOnlineResult<FGetUserInfo> GetUserInfo(FGetUserInfo::Params&& Params) = 0;
+
+	/*
+	 * Queries the user avatars for a list of account ids
+	 * Upon success, avatars can be accessed with GetUserAvatar.
+	 * @see GetUserAvatar
+	 *
+	 * @params Params for the QueryUserAvatar call
+	 * @return AsyncOpHandle
+	 */
+	virtual TOnlineAsyncOpHandle<FQueryUserAvatar> QueryUserAvatar(FQueryUserAvatar::Params&& Params) = 0;
+
+	/*
+	 * Get the user avatar for a previously queried AccountId
+	 * @see QueryUserAvatar
+	 *
+	 * @params Params for the GetUserAvatar call
+	 * @return Result
+	 */
+	virtual TOnlineResult<FGetUserAvatar> GetUserAvatar(FGetUserAvatar::Params&& Params) = 0;
+
+	/*
+	 * Show the profile UI for a user
+	 *
+	 * @params Params for the ShowUserProfile call
+	 * @return AsyncOpHandle
+	 */
+	virtual TOnlineAsyncOpHandle<FShowUserProfile> ShowUserProfile(FShowUserProfile::Params&& Params) = 0;
+};
+
+namespace Meta {
+
+BEGIN_ONLINE_STRUCT_META(FUserInfo)
+	ONLINE_STRUCT_FIELD(FUserInfo, UserId),
+	ONLINE_STRUCT_FIELD(FUserInfo, DisplayName)
+END_ONLINE_STRUCT_META()
+
+BEGIN_ONLINE_STRUCT_META(FQueryUserInfo::Params)
+	ONLINE_STRUCT_FIELD(FQueryUserInfo::Params, LocalUserId),
+	ONLINE_STRUCT_FIELD(FQueryUserInfo::Params, UserIds)
+END_ONLINE_STRUCT_META()
+
+BEGIN_ONLINE_STRUCT_META(FQueryUserInfo::Result)
+END_ONLINE_STRUCT_META()
+
+BEGIN_ONLINE_STRUCT_META(FGetUserInfo::Params)
+	ONLINE_STRUCT_FIELD(FGetUserInfo::Params, LocalUserId),
+	ONLINE_STRUCT_FIELD(FGetUserInfo::Params, UserId)
+END_ONLINE_STRUCT_META()
+
+BEGIN_ONLINE_STRUCT_META(FGetUserInfo::Result)
+	ONLINE_STRUCT_FIELD(FGetUserInfo::Result, UserInfo)
+END_ONLINE_STRUCT_META()
+
+BEGIN_ONLINE_STRUCT_META(FQueryUserAvatar::Params)
+	ONLINE_STRUCT_FIELD(FQueryUserAvatar::Params, LocalUserId),
+	ONLINE_STRUCT_FIELD(FQueryUserAvatar::Params, UserIds)
+END_ONLINE_STRUCT_META()
+
+BEGIN_ONLINE_STRUCT_META(FQueryUserAvatar::Result)
+END_ONLINE_STRUCT_META()
+
+BEGIN_ONLINE_STRUCT_META(FGetUserAvatar::Params)
+	ONLINE_STRUCT_FIELD(FGetUserAvatar::Params, LocalUserId),
+	ONLINE_STRUCT_FIELD(FGetUserAvatar::Params, UserId)
+END_ONLINE_STRUCT_META()
+
+BEGIN_ONLINE_STRUCT_META(FGetUserAvatar::Result)
+	ONLINE_STRUCT_FIELD(FGetUserAvatar::Result, AvatarUrl)
+END_ONLINE_STRUCT_META()
+
+BEGIN_ONLINE_STRUCT_META(FShowUserProfile::Params)
+	ONLINE_STRUCT_FIELD(FShowUserProfile::Params, LocalUserId),
+	ONLINE_STRUCT_FIELD(FShowUserProfile::Params, UserId)
+END_ONLINE_STRUCT_META()
+
+BEGIN_ONLINE_STRUCT_META(FShowUserProfile::Result)
+END_ONLINE_STRUCT_META()
+
+/* Meta*/ }
+
+/* UE::Online */ }
\ No newline at end of file
