Implement OSSv2 Auth Interface.
Implement OSSv2 AuthCommon.
Change LobbyVariant to SchemaVariant for more general use.
Add FOnlineVerifiedAuthTicketIdHandle and FOnlineVerifiedAuthSessionIdHandle

diff --git a/Engine/Plugins/Online/Experimental/OnlineServices/Source/OnlineServicesCommon/Private/Online/AuthCommon.cpp b/Engine/Plugins/Online/Experimental/OnlineServices/Source/OnlineServicesCommon/Private/Online/AuthCommon.cpp
--- a/Engine/Plugins/Online/Experimental/OnlineServices/Source/OnlineServicesCommon/Private/Online/AuthCommon.cpp
+++ b/Engine/Plugins/Online/Experimental/OnlineServices/Source/OnlineServicesCommon/Private/Online/AuthCommon.cpp
@@ -2,12 +2,72 @@
 
 #include "Online/AuthCommon.h"
 
+#include "Algo/Transform.h"
+#include "Misc/ScopeRWLock.h"
 #include "Online/OnlineAsyncOp.h"
 #include "Online/OnlineErrorDefinitions.h"
 #include "Online/OnlineServicesCommon.h"
 
 namespace UE::Online {
 
+TSharedPtr<FAccountInfo> FAccountInfoRegistry::Find(FPlatformUserId PlatformUserId) const
+{
+	FReadScopeLock Lock(IndexLock);
+	const TSharedRef<FAccountInfo>* FoundPtr = AuthDataByPlatformUserId.Find(PlatformUserId);
+	return FoundPtr ? TSharedPtr<FAccountInfo>(*FoundPtr) : TSharedPtr<FAccountInfo>();
+}
+
+TSharedPtr<FAccountInfo> FAccountInfoRegistry::Find(FOnlineAccountIdHandle AccountIdHandle) const
+{
+	FReadScopeLock Lock(IndexLock);
+	const TSharedRef<FAccountInfo>* FoundPtr = AuthDataByOnlineAccountIdHandle.Find(AccountIdHandle);
+	return FoundPtr ? TSharedPtr<FAccountInfo>(*FoundPtr) : TSharedPtr<FAccountInfo>();
+}
+
+TArray<TSharedRef<FAccountInfo>> FAccountInfoRegistry::GetAllAccountInfo(TFunction<bool(const TSharedRef<FAccountInfo>&)> Predicate) const
+{
+	FReadScopeLock Lock(IndexLock);
+	TArray<TSharedRef<FAccountInfo>> Result;
+	Result.Reserve(AuthDataByPlatformUserId.Num());
+	Algo::TransformIf(
+		AuthDataByPlatformUserId,
+		Result,
+		[Predicate](const TPair<FPlatformUserId, TSharedRef<FAccountInfo>>& Value)
+		{
+			return Predicate(Value.Get<1>());
+		},
+		[](const TPair<FPlatformUserId, TSharedRef<FAccountInfo>>& Value)
+		{
+			return Value.Get<1>();
+		});
+	return Result;
+}
+
+void FAccountInfoRegistry::DoRegister(const TSharedRef<FAccountInfo>& AccountInfo)
+{
+	if (ensure(AccountInfo->PlatformUserId != PLATFORMUSERID_NONE))
+	{
+		AuthDataByPlatformUserId.Add(AccountInfo->PlatformUserId, AccountInfo);
+	}
+
+	if (ensure(AccountInfo->AccountId.IsValid()))
+	{
+		AuthDataByOnlineAccountIdHandle.Add(AccountInfo->AccountId, AccountInfo);
+	}
+}
+void FAccountInfoRegistry::DoUnregister(const TSharedRef<FAccountInfo>& AccountInfo)
+{
+	if (ensure(AccountInfo->PlatformUserId != PLATFORMUSERID_NONE))
+	{
+		AuthDataByPlatformUserId.Remove(AccountInfo->PlatformUserId);
+	}
+
+	if (ensure(AccountInfo->AccountId.IsValid()))
+	{
+		AuthDataByOnlineAccountIdHandle.Remove(AccountInfo->AccountId);
+	}
+}
+
 FAuthCommon::FAuthCommon(FOnlineServicesCommon& InServices)
 	: TOnlineComponent(TEXT("Auth"), InServices)
 {
@@ -17,17 +77,22 @@ void FAuthCommon::RegisterCommands()
 {
 	RegisterCommand(&FAuthCommon::Login);
 	RegisterCommand(&FAuthCommon::Logout);
-	RegisterCommand(&FAuthCommon::GenerateAuthToken);
-	RegisterCommand(&FAuthCommon::GetAuthToken);
-	RegisterCommand(&FAuthCommon::GenerateAuthCode);
-	RegisterCommand(&FAuthCommon::GetAccountByPlatformUserId);
-	RegisterCommand(&FAuthCommon::GetAccountByAccountId);
+	RegisterCommand(&FAuthCommon::ModifyAccountAttributes);
+	RegisterCommand(&FAuthCommon::QueryExternalServerAuthTicket);
+	RegisterCommand(&FAuthCommon::QueryExternalAuthToken);
+	RegisterCommand(&FAuthCommon::QueryVerifiedAuthTicket);
+	RegisterCommand(&FAuthCommon::CancelVerifiedAuthTicket);
+	RegisterCommand(&FAuthCommon::BeginVerifiedAuthSession);
+	RegisterCommand(&FAuthCommon::EndVerifiedAuthSession);
+	RegisterCommand(&FAuthCommon::GetLocalOnlineUserByOnlineAccountId);
+	RegisterCommand(&FAuthCommon::GetLocalOnlineUserByPlatformUserId);
+	RegisterCommand(&FAuthCommon::GetAllLocalOnlineUsers);
 }
 
 TOnlineAsyncOpHandle<FAuthLogin> FAuthCommon::Login(FAuthLogin::Params&& Params)
 {
 	TOnlineAsyncOpRef<FAuthLogin> Operation = GetOp<FAuthLogin>(MoveTemp(Params));
-	Operation->SetError(Errors::NotImplemented());  
+	Operation->SetError(Errors::NotImplemented());
 	return Operation->GetHandle();
 }
 
@@ -38,38 +103,96 @@ TOnlineAsyncOpHandle<FAuthLogout> FAuthCommon::Logout(FAuthLogout::Params&& Para
 	return Operation->GetHandle();
 }
 
-TOnlineAsyncOpHandle<FAuthGenerateAuthToken> FAuthCommon::GenerateAuthToken(FAuthGenerateAuthToken::Params&& Params)
+TOnlineAsyncOpHandle<FAuthModifyAccountAttributes> FAuthCommon::ModifyAccountAttributes(FAuthModifyAccountAttributes::Params&& Params)
 {
-	TOnlineAsyncOpRef<FAuthGenerateAuthToken> Operation = GetOp<FAuthGenerateAuthToken>(MoveTemp(Params));
+	TOnlineAsyncOpRef<FAuthModifyAccountAttributes> Operation = GetOp<FAuthModifyAccountAttributes>(MoveTemp(Params));
 	Operation->SetError(Errors::NotImplemented());
 	return Operation->GetHandle();
 }
 
-TOnlineResult<FAuthGetAuthToken> FAuthCommon::GetAuthToken(FAuthGetAuthToken::Params&& Params)
+TOnlineAsyncOpHandle<FAuthQueryExternalServerAuthTicket> FAuthCommon::QueryExternalServerAuthTicket(FAuthQueryExternalServerAuthTicket::Params&& Params)
 {
-	return TOnlineResult<FAuthGetAuthToken>(Errors::NotImplemented());
+	TOnlineAsyncOpRef<FAuthQueryExternalServerAuthTicket> Operation = GetOp<FAuthQueryExternalServerAuthTicket>(MoveTemp(Params));
+	Operation->SetError(Errors::NotImplemented());
+	return Operation->GetHandle();
 }
 
-TOnlineAsyncOpHandle<FAuthGenerateAuthCode> FAuthCommon::GenerateAuthCode(FAuthGenerateAuthCode::Params&& Params)
+TOnlineAsyncOpHandle<FAuthQueryExternalAuthToken> FAuthCommon::QueryExternalAuthToken(FAuthQueryExternalAuthToken::Params&& Params)
 {
-	TOnlineAsyncOpRef<FAuthGenerateAuthCode> Operation = GetOp<FAuthGenerateAuthCode>(MoveTemp(Params));
+	TOnlineAsyncOpRef<FAuthQueryExternalAuthToken> Operation = GetOp<FAuthQueryExternalAuthToken>(MoveTemp(Params));
 	Operation->SetError(Errors::NotImplemented());
 	return Operation->GetHandle();
 }
 
-TOnlineResult<FAuthGetAccountByPlatformUserId> FAuthCommon::GetAccountByPlatformUserId(FAuthGetAccountByPlatformUserId::Params&& Params)
+TOnlineAsyncOpHandle<FAuthQueryVerifiedAuthTicket> FAuthCommon::QueryVerifiedAuthTicket(FAuthQueryVerifiedAuthTicket::Params&& Params)
+{
+	TOnlineAsyncOpRef<FAuthQueryVerifiedAuthTicket> Operation = GetOp<FAuthQueryVerifiedAuthTicket>(MoveTemp(Params));
+	Operation->SetError(Errors::NotImplemented());
+	return Operation->GetHandle();
+}
+
+TOnlineAsyncOpHandle<FAuthCancelVerifiedAuthTicket> FAuthCommon::CancelVerifiedAuthTicket(FAuthCancelVerifiedAuthTicket::Params&& Params)
+{
+	TOnlineAsyncOpRef<FAuthCancelVerifiedAuthTicket> Operation = GetOp<FAuthCancelVerifiedAuthTicket>(MoveTemp(Params));
+	Operation->SetError(Errors::NotImplemented());
+	return Operation->GetHandle();
+}
+
+TOnlineAsyncOpHandle<FAuthBeginVerifiedAuthSession> FAuthCommon::BeginVerifiedAuthSession(FAuthBeginVerifiedAuthSession::Params&& Params)
+{
+	TOnlineAsyncOpRef<FAuthBeginVerifiedAuthSession> Operation = GetOp<FAuthBeginVerifiedAuthSession>(MoveTemp(Params));
+	Operation->SetError(Errors::NotImplemented());
+	return Operation->GetHandle();
+}
+
+TOnlineAsyncOpHandle<FAuthEndVerifiedAuthSession> FAuthCommon::EndVerifiedAuthSession(FAuthEndVerifiedAuthSession::Params&& Params)
+{
+	TOnlineAsyncOpRef<FAuthEndVerifiedAuthSession> Operation = GetOp<FAuthEndVerifiedAuthSession>(MoveTemp(Params));
+	Operation->SetError(Errors::NotImplemented());
+	return Operation->GetHandle();
+}
+
+TOnlineResult<FAuthGetLocalOnlineUserByPlatformUserId> FAuthCommon::GetLocalOnlineUserByPlatformUserId(FAuthGetLocalOnlineUserByPlatformUserId::Params&& Params) const
+{
+	TSharedPtr<FAccountInfo> AccountInfo = GetAccountInfoRegistry().Find(Params.PlatformUserId);
+	return (AccountInfo && IsOnlineStatus(AccountInfo->LoginStatus)) ?
+		TOnlineResult<FAuthGetLocalOnlineUserByPlatformUserId>(FAuthGetLocalOnlineUserByPlatformUserId::Result{ AccountInfo.ToSharedRef() }) :
+		TOnlineResult<FAuthGetLocalOnlineUserByPlatformUserId>(Errors::InvalidUser());
+}
+
+TOnlineResult<FAuthGetLocalOnlineUserByOnlineAccountId> FAuthCommon::GetLocalOnlineUserByOnlineAccountId(FAuthGetLocalOnlineUserByOnlineAccountId::Params&& Params) const
+{
+	TSharedPtr<FAccountInfo> AccountInfo = GetAccountInfoRegistry().Find(Params.LocalUserId);
+	return (AccountInfo && IsOnlineStatus(AccountInfo->LoginStatus)) ?
+		TOnlineResult<FAuthGetLocalOnlineUserByOnlineAccountId>(FAuthGetLocalOnlineUserByOnlineAccountId::Result{ AccountInfo.ToSharedRef() }) :
+		TOnlineResult<FAuthGetLocalOnlineUserByOnlineAccountId>(Errors::InvalidUser());
+}
+
+TOnlineResult<FAuthGetAllLocalOnlineUsers> FAuthCommon::GetAllLocalOnlineUsers(FAuthGetAllLocalOnlineUsers::Params&& Params) const
+{
+	return TOnlineResult<FAuthGetAllLocalOnlineUsers>(FAuthGetAllLocalOnlineUsers::Result{
+		GetAccountInfoRegistry().GetAllAccountInfo([](const TSharedRef<FAccountInfo>& AccountInfo) { return IsOnlineStatus(AccountInfo->LoginStatus); }) });
+}
+
+TOnlineEvent<void(const FAuthLoginStatusChanged&)> FAuthCommon::OnLoginStatusChanged()
+{
+	return OnAuthLoginStatusChangedEvent;
+}
+
+TOnlineEvent<void(const FAuthPendingAuthExpiration&)> FAuthCommon::OnPendingAuthExpiration()
 {
-	return TOnlineResult<FAuthGetAccountByPlatformUserId>(Errors::NotImplemented());
+	return OnAuthPendingAuthExpirationEvent;
 }
 
-TOnlineResult<FAuthGetAccountByAccountId> FAuthCommon::GetAccountByAccountId(FAuthGetAccountByAccountId::Params&& Params)
+TOnlineEvent<void(const FAuthAccountAttributesChanged&)> FAuthCommon::OnAccountAttributesChanged()
 {
-	return TOnlineResult<FAuthGetAccountByAccountId>(Errors::NotImplemented());
+	return OnAuthAccountAttributesChangedEvent;
 }
 
-TOnlineEvent<void(const FLoginStatusChanged&)> FAuthCommon::OnLoginStatusChanged()
+bool FAuthCommon::IsLoggedIn(const FOnlineAccountIdHandle& AccountId) const
 {
-	return OnLoginStatusChangedEvent;
+	const TSharedPtr<FAccountInfo> AccountInfo = GetAccountInfoRegistry().Find(AccountId);
+	return (AccountInfo && IsOnlineStatus(AccountInfo->LoginStatus));
 }
 
 /* UE::Online */ }
diff --git a/Engine/Plugins/Online/Experimental/OnlineServices/Source/OnlineServicesCommon/Public/Online/AuthCommon.h b/Engine/Plugins/Online/Experimental/OnlineServices/Source/OnlineServicesCommon/Public/Online/AuthCommon.h
--- a/Engine/Plugins/Online/Experimental/OnlineServices/Source/OnlineServicesCommon/Public/Online/AuthCommon.h
+++ b/Engine/Plugins/Online/Experimental/OnlineServices/Source/OnlineServicesCommon/Public/Online/AuthCommon.h
@@ -9,6 +9,27 @@ namespace UE::Online {
 
 class FOnlineServicesCommon;
 
+class ONLINESERVICESCOMMON_API FAccountInfoRegistry
+{
+public:
+	virtual ~FAccountInfoRegistry() = default;
+
+	TSharedPtr<FAccountInfo> Find(FPlatformUserId PlatformUserId) const;
+	TSharedPtr<FAccountInfo> Find(FOnlineAccountIdHandle AccountIdHandle) const;
+
+	TArray<TSharedRef<FAccountInfo>> GetAllAccountInfo(TFunction<bool(const TSharedRef<FAccountInfo>&)> Predicate) const;
+
+protected:
+	virtual void DoRegister(const TSharedRef<FAccountInfo>& AccountInfo);
+	virtual void DoUnregister(const TSharedRef<FAccountInfo>& AccountInfo);
+
+	mutable FRWLock IndexLock;
+
+private:
+	TMap<FPlatformUserId, TSharedRef<FAccountInfo>> AuthDataByPlatformUserId;
+	TMap<FOnlineAccountIdHandle, TSharedRef<FAccountInfo>> AuthDataByOnlineAccountIdHandle;
+};
+
 class ONLINESERVICESCOMMON_API FAuthCommon : public TOnlineComponent<IAuth>
 {
 public:
@@ -22,15 +43,27 @@ public:
 	// IAuth
 	virtual TOnlineAsyncOpHandle<FAuthLogin> Login(FAuthLogin::Params&& Params) override;
 	virtual TOnlineAsyncOpHandle<FAuthLogout> Logout(FAuthLogout::Params&& Params) override;
-	virtual TOnlineAsyncOpHandle<FAuthGenerateAuthToken> GenerateAuthToken(FAuthGenerateAuthToken::Params&& Params) override;
-	virtual TOnlineResult<FAuthGetAuthToken> GetAuthToken(FAuthGetAuthToken::Params&& Params) override;
-	virtual TOnlineAsyncOpHandle<FAuthGenerateAuthCode> GenerateAuthCode(FAuthGenerateAuthCode::Params&& Params) override;
-	virtual TOnlineResult<FAuthGetAccountByPlatformUserId> GetAccountByPlatformUserId(FAuthGetAccountByPlatformUserId::Params&& Params) override;
-	virtual TOnlineResult<FAuthGetAccountByAccountId> GetAccountByAccountId(FAuthGetAccountByAccountId::Params&& Params) override;
-	virtual TOnlineEvent<void(const FLoginStatusChanged&)> OnLoginStatusChanged() override;
+	virtual TOnlineAsyncOpHandle<FAuthModifyAccountAttributes> ModifyAccountAttributes(FAuthModifyAccountAttributes::Params&& Params) override;
+	virtual TOnlineAsyncOpHandle<FAuthQueryExternalServerAuthTicket> QueryExternalServerAuthTicket(FAuthQueryExternalServerAuthTicket::Params&& Params) override;
+	virtual TOnlineAsyncOpHandle<FAuthQueryExternalAuthToken> QueryExternalAuthToken(FAuthQueryExternalAuthToken::Params&& Params) override;
+	virtual TOnlineAsyncOpHandle<FAuthQueryVerifiedAuthTicket> QueryVerifiedAuthTicket(FAuthQueryVerifiedAuthTicket::Params&& Params) override;
+	virtual TOnlineAsyncOpHandle<FAuthCancelVerifiedAuthTicket> CancelVerifiedAuthTicket(FAuthCancelVerifiedAuthTicket::Params&& Params) override;
+	virtual TOnlineAsyncOpHandle<FAuthBeginVerifiedAuthSession> BeginVerifiedAuthSession(FAuthBeginVerifiedAuthSession::Params&& Params) override;
+	virtual TOnlineAsyncOpHandle<FAuthEndVerifiedAuthSession> EndVerifiedAuthSession(FAuthEndVerifiedAuthSession::Params&& Params) override;
+	virtual TOnlineResult<FAuthGetLocalOnlineUserByOnlineAccountId> GetLocalOnlineUserByOnlineAccountId(FAuthGetLocalOnlineUserByOnlineAccountId::Params&& Params) const override;
+	virtual TOnlineResult<FAuthGetLocalOnlineUserByPlatformUserId> GetLocalOnlineUserByPlatformUserId(FAuthGetLocalOnlineUserByPlatformUserId::Params&& Params) const override;
+	virtual TOnlineResult<FAuthGetAllLocalOnlineUsers> GetAllLocalOnlineUsers(FAuthGetAllLocalOnlineUsers::Params&& Params) const override;
+	virtual TOnlineEvent<void(const FAuthLoginStatusChanged&)> OnLoginStatusChanged() override;
+	virtual TOnlineEvent<void(const FAuthPendingAuthExpiration&)> OnPendingAuthExpiration() override;
+	virtual TOnlineEvent<void(const FAuthAccountAttributesChanged&)> OnAccountAttributesChanged() override;
+	virtual bool IsLoggedIn(const FOnlineAccountIdHandle& AccountId) const override;
 
 protected:
-	TOnlineEventCallable<void(const FLoginStatusChanged&)> OnLoginStatusChangedEvent;
+	virtual const FAccountInfoRegistry& GetAccountInfoRegistry() const = 0;
+
+	TOnlineEventCallable<void(const FAuthLoginStatusChanged&)> OnAuthLoginStatusChangedEvent;
+	TOnlineEventCallable<void(const FAuthPendingAuthExpiration&)> OnAuthPendingAuthExpirationEvent;
+	TOnlineEventCallable<void(const FAuthAccountAttributesChanged&)> OnAuthAccountAttributesChangedEvent;
 };
 
 /* UE::Online */ }
diff --git a/Engine/Plugins/Online/Experimental/OnlineServices/Source/OnlineServicesCommon/Public/Online/OnlineExecHandler.h b/Engine/Plugins/Online/Experimental/OnlineServices/Source/OnlineServicesCommon/Public/Online/OnlineExecHandler.h
--- a/Engine/Plugins/Online/Experimental/OnlineServices/Source/OnlineServicesCommon/Public/Online/OnlineExecHandler.h
+++ b/Engine/Plugins/Online/Experimental/OnlineServices/Source/OnlineServicesCommon/Public/Online/OnlineExecHandler.h
@@ -70,6 +70,7 @@
 inline bool ParseOnlineExecParams(const TCHAR*& Cmd, bool& Value);
 template <typename T> inline bool ParseOnlineExecParams(const TCHAR*& Cmd, TArray<T>& Array);
 template <typename T, typename U> inline bool ParseOnlineExecParams(const TCHAR*& Cmd, TMap<T, U>& Map);
+template <typename T> inline bool ParseOnlineExecParams(const TCHAR*& Cmd, TSet<T>& Set);
 template <typename T> inline bool ParseOnlineExecParams(const TCHAR*& Cmd, TSharedPtr<T>& Ptr);
 template <typename T> inline bool ParseOnlineExecParams(const TCHAR*& Cmd, TSharedRef<T>& Ref);
 template <typename T> inline bool ParseOnlineExecParams(const TCHAR*& Cmd, TOptional<T>& Optional);
@@ -271,6 +272,32 @@
 	return true;
 }
 
+template <typename T>
+inline bool ParseOnlineExecParams(const TCHAR*& Cmd, TSet<T>& Set)
+{
+	FString Token;
+	if (FParse::Token(Cmd, Token, true))
+	{
+		TArray<FString> TokenArray;
+		Token.ParseIntoArray(TokenArray, TEXT(","));
+		Set.Reserve(TokenArray.Num());
+		for (const FString& ArrayToken : TokenArray)
+		{
+			T Key;
+
+			const TCHAR* ArrayTokenKeyTCHAR = *ArrayToken;
+			if (ParseOnlineExecParams(ArrayTokenKeyTCHAR, Key))
+			{
+				Set.Emplace(MoveTemp(Key));
+			}
+		}
+	}
+	else
+	{
+		return false;
+	}
+	return true;
+}
 
 template <typename T>
 inline bool ParseOnlineExecParams(const TCHAR*& Cmd, TSharedPtr<T>& Ptr)
diff --git a/Engine/Plugins/Online/Experimental/OnlineServices/Source/OnlineServicesInterface/Private/Online/Auth.cpp b/Engine/Plugins/Online/Experimental/OnlineServices/Source/OnlineServicesInterface/Private/Online/Auth.cpp
--- a/Engine/Plugins/Online/Experimental/OnlineServices/Source/OnlineServicesInterface/Private/Online/Auth.cpp
+++ b/Engine/Plugins/Online/Experimental/OnlineServices/Source/OnlineServicesInterface/Private/Online/Auth.cpp
@@ -4,12 +4,58 @@
 
 namespace UE::Online {
 
+namespace LoginCredentialsType
+{
+const FName Auto = TEXT("Auto");
+const FName Password = TEXT("Password");
+const FName ExchangeCode = TEXT("ExchangeCode");
+const FName PersistentAuth = TEXT("PersistentAuth");
+const FName Developer = TEXT("Developer");
+const FName RefreshToken = TEXT("RefreshToken");
+const FName AccountPortal = TEXT("AccountPortal");
+const FName ExternalAuth = TEXT("ExternalAuth");
+}
+
+namespace ExternalLoginType
+{
+const FName Epic = TEXT("Epic");
+const FName SteamAppTicket = TEXT("SteamAppTicket");
+const FName PsnIdToken = TEXT("PsnIdToken");
+const FName XblXstsToken = TEXT("XblXstsToken");
+const FName DiscordAccessToken = TEXT("DiscordAccessToken");
+const FName GogSessionTicket = TEXT("GogSessionTicket");
+const FName NintendoIdToken = TEXT("NintendoIdToken");
+const FName NintendoNsaIdToken = TEXT("NintendoNsaIdToken");
+const FName UplayAccessToken = TEXT("UplayAccessToken");
+const FName OpenIdAccessToken = TEXT("OpenIdAccessToken");
+const FName DeviceIdAccessToken = TEXT("DeviceIdAccessToken");
+const FName AppleIdToken = TEXT("AppleIdToken");
+const FName GoogleIdToken = TEXT("GoogleIdToken");
+const FName OculusUserIdNonce = TEXT("OculusUserIdNonce");
+const FName ItchioJwt = TEXT("ItchioJwt");
+const FName ItchioKey = TEXT("ItchioKey");
+const FName EpicIdToken = TEXT("EpicIdToken");
+const FName AmazonAccessToken = TEXT("AmazonAccessToken");
+}
+
+namespace ExternalServerAuthTicketType
+{
+const FName PsnAuthCode = TEXT("PsnAuthCode");
+const FName XblXstsToken = TEXT("XblXstsToken");
+}
+
+namespace AccountAttributeData
+{
+const FSchemaAttributeId DisplayName = TEXT("DisplayName");
+}
+
 const TCHAR* LexToString(ELoginStatus Status)
 {
 	switch (Status)
 	{
 	case ELoginStatus::UsingLocalProfile:	return TEXT("UsingLocalProfile");
 	case ELoginStatus::LoggedIn:			return TEXT("LoggedIn");
+	case ELoginStatus::LoggedInReducedFunctionality:	return TEXT("LoggedInReducedFunctionality");
 	default:								checkNoEntry(); // Intentional fallthrough
 	case ELoginStatus::NotLoggedIn:			return TEXT("NotLoggedIn");
 	}
@@ -25,6 +71,10 @@ void LexFromString(ELoginStatus& OutStatus, const TCHAR* InStr)
 	{
 		OutStatus = ELoginStatus::UsingLocalProfile;
 	}
+	else if (FCString::Stricmp(InStr, TEXT("LoggedInReducedFunctionality")) == 0)
+	{
+		OutStatus = ELoginStatus::LoggedInReducedFunctionality;
+	}
 	else if (FCString::Stricmp(InStr, TEXT("NotLoggedIn")) == 0)
 	{
 		OutStatus = ELoginStatus::NotLoggedIn;
@@ -36,4 +86,58 @@ void LexFromString(ELoginStatus& OutStatus, const TCHAR* InStr)
 	}
 }
 
+const TCHAR* LexToString(ERemoteAuthTicketAudience Audience)
+{
+	switch (Audience)
+	{
+	case ERemoteAuthTicketAudience::DedicatedServer:	return TEXT("DedicatedServer");
+	default:											checkNoEntry(); // Intentional fallthrough
+	case ERemoteAuthTicketAudience::Peer:				return TEXT("Peer");
+	}
+}
+
+void LexFromString(ERemoteAuthTicketAudience& OutAudience, const TCHAR* InStr)
+{
+	if (FCString::Stricmp(InStr, TEXT("Peer")) == 0)
+	{
+		OutAudience = ERemoteAuthTicketAudience::Peer;
+	}
+	else if (FCString::Stricmp(InStr, TEXT("DedicatedServer")) == 0)
+	{
+		OutAudience = ERemoteAuthTicketAudience::DedicatedServer;
+	}
+	else
+	{
+		checkNoEntry();
+		OutAudience = ERemoteAuthTicketAudience::Peer;
+	}
+}
+
+const TCHAR* LexToString(EExternalAuthTokenMethod Method)
+{
+	switch (Method)
+	{
+	case EExternalAuthTokenMethod::Primary:		return TEXT("Primary");
+	default:									checkNoEntry(); // Intentional fallthrough
+	case EExternalAuthTokenMethod::Secondary:	return TEXT("Secondary");
+	}
+}
+
+void LexFromString(EExternalAuthTokenMethod& OutMethod, const TCHAR* InStr)
+{
+	if (FCString::Stricmp(InStr, TEXT("Primary")) == 0)
+	{
+		OutMethod = EExternalAuthTokenMethod::Primary;
+	}
+	else if (FCString::Stricmp(InStr, TEXT("Secondary")) == 0)
+	{
+		OutMethod = EExternalAuthTokenMethod::Secondary;
+	}
+	else
+	{
+		checkNoEntry();
+		OutMethod = EExternalAuthTokenMethod::Primary;
+	}
+}
+
 /* UE::Online */ }
diff --git a/Engine/Plugins/Online/Experimental/OnlineServices/Source/OnlineServicesInterface/Private/Online/Lobbies.cpp b/Engine/Plugins/Online/Experimental/OnlineServices/Source/OnlineServicesInterface/Private/Online/Lobbies.cpp
--- a/Engine/Plugins/Online/Experimental/OnlineServices/Source/OnlineServicesInterface/Private/Online/Lobbies.cpp
+++ b/Engine/Plugins/Online/Experimental/OnlineServices/Source/OnlineServicesInterface/Private/Online/Lobbies.cpp
@@ -4,121 +4,6 @@
 
 namespace UE::Online {
 
-const FString LexToString(const FLobbyVariant& Variant)
-{
-	return Variant.GetString();
-}
-
-void LexFromString(FLobbyVariant& Variant, const TCHAR* InStr)
-{
-	Variant.Set(TEXT("")); // todo
-}
-
-int64 FLobbyVariant::GetInt64() const
-{
-	int64 AsInt64 = 0;
-	if (VariantData.IsType<int64>())
-	{
-		AsInt64 = VariantData.Get<int64>();
-	}
-	else if (VariantData.IsType<bool>())
-	{
-		AsInt64 = static_cast<int64>(VariantData.Get<bool>());
-	}
-	else if (VariantData.IsType<double>())
-	{
-		AsInt64 = static_cast<int64>(VariantData.Get<double>());
-	}
-	else if (VariantData.IsType<FString>())
-	{
-		::LexFromString(AsInt64, *VariantData.Get<FString>());
-	}
-	return AsInt64;
-}
-
-double FLobbyVariant::GetDouble() const
-{
-	double AsDouble = 0;
-	if (VariantData.IsType<double>())
-	{
-		AsDouble = VariantData.Get<double>();
-	}
-	else if (VariantData.IsType<FString>())
-	{
-		::LexFromString(AsDouble, *VariantData.Get<FString>());
-	}
-	else if (VariantData.IsType<int64>())
-	{
-		AsDouble = static_cast<double>(VariantData.Get<int64>());
-	}
-	else if (VariantData.IsType<bool>())
-	{
-		AsDouble = static_cast<double>(VariantData.Get<bool>());
-	}
-	return AsDouble;
-}
-
-bool FLobbyVariant::GetBoolean() const
-{
-	bool bAsBool = false;
-	if (VariantData.IsType<bool>())
-	{
-		bAsBool = VariantData.Get<bool>();
-	}
-	else if (VariantData.IsType<FString>())
-	{
-		::LexFromString(bAsBool, *VariantData.Get<FString>());
-	}
-	else
-	{
-		bAsBool = GetInt64() != 0;
-	}
-	return bAsBool;
-}
-
-FString FLobbyVariant::GetString() const
-{
-	if (VariantData.IsType<FString>())
-	{
-		return VariantData.Get<FString>();
-	}
-	else if (VariantData.IsType<int64>())
-	{
-		return FString::Printf(TEXT("%" INT64_FMT), VariantData.Get<int64>());
-	}
-	else if (VariantData.IsType<bool>())
-	{
-		return ::LexToString(VariantData.Get<bool>());
-	}
-	else if (VariantData.IsType<double>())
-	{
-		return FString::Printf(TEXT("%f"), VariantData.Get<double>());
-	}
-	else
-	{
-		checkNoEntry();
-	}
-	return TEXT("");
-}
-
-bool FLobbyVariant::operator==(const FLobbyVariant& Other) const
-{
-	if (VariantData.GetIndex() != Other.VariantData.GetIndex())
-	{
-		return false;
-	}
-
-	switch (VariantData.GetIndex())
-	{
-	case FVariantType::IndexOfType<FString>():	return VariantData.Get<FString>() == Other.VariantData.Get<FString>();
-	case FVariantType::IndexOfType<int64>():	return VariantData.Get<int64>() == Other.VariantData.Get<int64>();
-	case FVariantType::IndexOfType<double>():	return VariantData.Get<double>() == Other.VariantData.Get<double>();
-
-	default:checkNoEntry(); // Intentional fallthrough
-	case FVariantType::IndexOfType<bool>():	return VariantData.Get<bool>() == Other.VariantData.Get<bool>();
-	}
-}
-
 const TCHAR* LexToString(ELobbyJoinPolicy Policy)
 {
 	switch (Policy)
diff --git a/Engine/Plugins/Online/Experimental/OnlineServices/Source/OnlineServicesInterface/Private/Online/Schema.cpp b/Engine/Plugins/Online/Experimental/OnlineServices/Source/OnlineServicesInterface/Private/Online/Schema.cpp
new file mode 100644
--- /dev/null
+++ b/Engine/Plugins/Online/Experimental/OnlineServices/Source/OnlineServicesInterface/Private/Online/Schema.cpp
@@ -0,0 +1,123 @@
+// Copyright Epic Games, Inc. All Rights Reserved.
+
+#include "Online/Schema.h"
+
+namespace UE::Online {
+
+
+const FString LexToString(const FSchemaVariant& Variant)
+{
+	return Variant.GetString();
+}
+
+void LexFromString(FSchemaVariant& Variant, const TCHAR* InStr)
+{
+	Variant.Set(TEXT("")); // todo
+}
+
+int64 FSchemaVariant::GetInt64() const
+{
+	int64 AsInt64 = 0;
+	if (VariantData.IsType<int64>())
+	{
+		AsInt64 = VariantData.Get<int64>();
+	}
+	else if (VariantData.IsType<bool>())
+	{
+		AsInt64 = static_cast<int64>(VariantData.Get<bool>());
+	}
+	else if (VariantData.IsType<double>())
+	{
+		AsInt64 = static_cast<int64>(VariantData.Get<double>());
+	}
+	else if (VariantData.IsType<FString>())
+	{
+		::LexFromString(AsInt64, *VariantData.Get<FString>());
+	}
+	return AsInt64;
+}
+
+double FSchemaVariant::GetDouble() const
+{
+	double AsDouble = 0;
+	if (VariantData.IsType<double>())
+	{
+		AsDouble = VariantData.Get<double>();
+	}
+	else if (VariantData.IsType<FString>())
+	{
+		::LexFromString(AsDouble, *VariantData.Get<FString>());
+	}
+	else if (VariantData.IsType<int64>())
+	{
+		AsDouble = static_cast<double>(VariantData.Get<int64>());
+	}
+	else if (VariantData.IsType<bool>())
+	{
+		AsDouble = static_cast<double>(VariantData.Get<bool>());
+	}
+	return AsDouble;
+}
+
+bool FSchemaVariant::GetBoolean() const
+{
+	bool bAsBool = false;
+	if (VariantData.IsType<bool>())
+	{
+		bAsBool = VariantData.Get<bool>();
+	}
+	else if (VariantData.IsType<FString>())
+	{
+		::LexFromString(bAsBool, *VariantData.Get<FString>());
+	}
+	else
+	{
+		bAsBool = GetInt64() != 0;
+	}
+	return bAsBool;
+}
+
+FString FSchemaVariant::GetString() const
+{
+	if (VariantData.IsType<FString>())
+	{
+		return VariantData.Get<FString>();
+	}
+	else if (VariantData.IsType<int64>())
+	{
+		return FString::Printf(TEXT("%" INT64_FMT), VariantData.Get<int64>());
+	}
+	else if (VariantData.IsType<bool>())
+	{
+		return ::LexToString(VariantData.Get<bool>());
+	}
+	else if (VariantData.IsType<double>())
+	{
+		return FString::Printf(TEXT("%f"), VariantData.Get<double>());
+	}
+	else
+	{
+		checkNoEntry();
+	}
+	return TEXT("");
+}
+
+bool FSchemaVariant::operator==(const FSchemaVariant& Other) const
+{
+	if (VariantData.GetIndex() != Other.VariantData.GetIndex())
+	{
+		return false;
+	}
+
+	switch (VariantData.GetIndex())
+	{
+	case FVariantType::IndexOfType<FString>():	return VariantData.Get<FString>() == Other.VariantData.Get<FString>();
+	case FVariantType::IndexOfType<int64>():	return VariantData.Get<int64>() == Other.VariantData.Get<int64>();
+	case FVariantType::IndexOfType<double>():	return VariantData.Get<double>() == Other.VariantData.Get<double>();
+
+	default:checkNoEntry(); // Intentional fallthrough
+	case FVariantType::IndexOfType<bool>():	return VariantData.Get<bool>() == Other.VariantData.Get<bool>();
+	}
+}
+
+/* UE::Online */ }
diff --git a/Engine/Plugins/Online/Experimental/OnlineServices/Source/OnlineServicesInterface/Public/Online/Auth.h b/Engine/Plugins/Online/Experimental/OnlineServices/Source/OnlineServicesInterface/Public/Online/Auth.h
--- a/Engine/Plugins/Online/Experimental/OnlineServices/Source/OnlineServicesInterface/Public/Online/Auth.h
+++ b/Engine/Plugins/Online/Experimental/OnlineServices/Source/OnlineServicesInterface/Public/Online/Auth.h
@@ -2,41 +2,155 @@
 
 #pragma once
 
+#include "Containers/Map.h"
+#include "Containers/Set.h"
+
 #include "Online/CoreOnline.h"
 #include "Online/OnlineAsyncOpHandle.h"
 #include "Online/OnlineMeta.h"
+#include "Online/Schema.h"
 
 namespace UE::Online {
 
 class FOnlineError;
 
+namespace LoginCredentialsType
+{
+ONLINESERVICESINTERFACE_API extern const FName Auto;
+ONLINESERVICESINTERFACE_API extern const FName Password;
+ONLINESERVICESINTERFACE_API extern const FName ExchangeCode;
+ONLINESERVICESINTERFACE_API extern const FName PersistentAuth;
+ONLINESERVICESINTERFACE_API extern const FName Developer;
+ONLINESERVICESINTERFACE_API extern const FName RefreshToken;
+ONLINESERVICESINTERFACE_API extern const FName AccountPortal;
+ONLINESERVICESINTERFACE_API extern const FName ExternalAuth;
+}
+
+namespace ExternalLoginType
+{
+ONLINESERVICESINTERFACE_API extern const FName Epic;
+ONLINESERVICESINTERFACE_API extern const FName SteamAppTicket;
+ONLINESERVICESINTERFACE_API extern const FName PsnIdToken;
+ONLINESERVICESINTERFACE_API extern const FName XblXstsToken;
+ONLINESERVICESINTERFACE_API extern const FName DiscordAccessToken;
+ONLINESERVICESINTERFACE_API extern const FName GogSessionTicket;
+ONLINESERVICESINTERFACE_API extern const FName NintendoIdToken;
+ONLINESERVICESINTERFACE_API extern const FName NintendoNsaIdToken;
+ONLINESERVICESINTERFACE_API extern const FName UplayAccessToken;
+ONLINESERVICESINTERFACE_API extern const FName OpenIdAccessToken;
+ONLINESERVICESINTERFACE_API extern const FName DeviceIdAccessToken;
+ONLINESERVICESINTERFACE_API extern const FName AppleIdToken;
+ONLINESERVICESINTERFACE_API extern const FName GoogleIdToken;
+ONLINESERVICESINTERFACE_API extern const FName OculusUserIdNonce;
+ONLINESERVICESINTERFACE_API extern const FName ItchioJwt;
+ONLINESERVICESINTERFACE_API extern const FName ItchioKey;
+ONLINESERVICESINTERFACE_API extern const FName EpicIdToken;
+ONLINESERVICESINTERFACE_API extern const FName AmazonAccessToken;
+}
+
+namespace ExternalServerAuthTicketType
+{
+ONLINESERVICESINTERFACE_API extern const FName PsnAuthCode;
+ONLINESERVICESINTERFACE_API extern const FName XblXstsToken;
+}
+
+namespace AccountAttributeData
+{
+ONLINESERVICESINTERFACE_API extern const FSchemaAttributeId DisplayName;
+}
+
 enum class ELoginStatus : uint8
 {
-	/** Player has not logged in or chosen a local profile */
+	/** Player has not logged in or chosen a local profile. */
 	NotLoggedIn,
-	/** Player is using a local profile but is not logged in */
+	/** Player is using a local profile but is not logged in. */
 	UsingLocalProfile,
-	/** Player has been validated by the platform specific authentication service */
+	/** Player is logged in but may have reduced functionality with online services. */
+	LoggedInReducedFunctionality,
+	/** Player has been validated by the platform specific authentication service. */
 	LoggedIn
 };
 ONLINESERVICESINTERFACE_API const TCHAR* LexToString(ELoginStatus Status);
 ONLINESERVICESINTERFACE_API void LexFromString(ELoginStatus& OutStatus, const TCHAR* InStr);
 
-class FAccountInfo
+inline bool IsOnlineStatus(ELoginStatus LoginStatus)
 {
-public:
-	/** Local user num */
-	FPlatformUserId PlatformUserId;
-	/** Account id */
-	FOnlineAccountIdHandle UserId;
+	return LoginStatus == ELoginStatus::LoggedIn || LoginStatus == ELoginStatus::LoggedInReducedFunctionality;
+}
+
+enum class ERemoteAuthTicketAudience : uint8
+{
+	/** Generate a ticket appropriate for peer validation. */
+	Peer,
+	/**
+	 * Generate a ticket appropriate for dedicated server validation.
+	 * Depending on the platform, dedicated servers may have access to additional APIs used to verify tickets.
+	 */
+	DedicatedServer,
+};
+ONLINESERVICESINTERFACE_API const TCHAR* LexToString(ERemoteAuthTicketAudience Audience);
+ONLINESERVICESINTERFACE_API void LexFromString(ERemoteAuthTicketAudience& OutAudience, const TCHAR* InStr);
+
+/** Some auth interfaces have more than one method for providing credentials when linking to an
+ *  external account. An example usage is when the auth interface can provide a token for linking
+ *  the local hardware device as the primary method before falling back to a token linked to the
+ *  users online account.
+ */
+enum class EExternalAuthTokenMethod : uint8
+{
+	/** Acquire an external auth token using the primary method provided by the auth interface. */
+	Primary,
+	/** Acquire an external auth token using the secondary method provided by the auth interface. */
+	Secondary,
+};
+ONLINESERVICESINTERFACE_API const TCHAR* LexToString(EExternalAuthTokenMethod Method);
+ONLINESERVICESINTERFACE_API void LexFromString(EExternalAuthTokenMethod& OutMethod, const TCHAR* InStr);
+
+struct FAccountInfo
+{
+	/** The account id for the user which represents the user's online platform account. */
+	FOnlineAccountIdHandle AccountId;
+	/** The platform user id associated with the online user. */
+	FPlatformUserId PlatformUserId = PLATFORMUSERID_NONE;
 	/** Login status */
-	ELoginStatus LoginStatus;
-	/** Display name */
-	FString DisplayName;
-	// TODO: Other fields
+	ELoginStatus LoginStatus = ELoginStatus::NotLoggedIn;
+	/** Additional account attributes. */
+	TMap<FSchemaAttributeId, FSchemaVariant> Attributes;
 };
 
-using FCredentialsToken = TVariant<FString, TArray<uint8>>;
+/** Token appropriate for allowing a trusted server to verify authentication.
+ *  In most implementations this token will be an OpenId token.
+ */
+struct FExternalAuthToken
+{
+	FName Type;
+	FString Data;
+};
+
+/** Ticket appropriate for allowing a trusted server to access service APIs on the users behalf. */
+struct FExternalServerAuthTicket
+{
+	FName Type;
+	FString Data;
+};
+
+/** Token used during Login operation. The data provided as the credentials token varies based
+ * on the type of login being performed.
+ */
+using FCredentialsToken = TVariant<FString, FExternalAuthToken>;
+
+struct FVerifiedAuthSession
+{
+	FOnlineVerifiedAuthSessionIdHandle SessionId;
+	FOnlineAccountIdHandle RemoteAccountId;
+	double CreationTime;
+};
+
+struct FVerifiedAuthTicket
+{
+	FName Type;
+	FString Data;
+};
 
 struct FAuthLogin
 {
@@ -44,10 +158,11 @@ struct FAuthLogin
 
 	struct Params
 	{
+		/** The PlatformUserId of the Local User making the request. */
 		FPlatformUserId PlatformUserId = PLATFORMUSERID_NONE;
-		FString CredentialsType;
+		FName CredentialsType;
 		FString CredentialsId;
-		FCredentialsToken CredentialsToken;
+		TVariant<FString, FExternalAuthToken> CredentialsToken;
 		TArray<FString> Scopes;
 	};
 
@@ -66,7 +181,9 @@ struct FAuthLogout
 
 	struct Params
 	{
+		/** The online account id of the Local User making the request. */
 		FOnlineAccountIdHandle LocalUserId;
+		/** Whether to remove persistent credentials on logout. */
 		bool bDestroyAuth = false;
 	};
 
@@ -75,151 +192,362 @@ struct FAuthLogout
 	};
 };
 
-struct FAuthGenerateAuthToken
+struct FAuthModifyAccountAttributes
 {
-	static constexpr TCHAR Name[] = TEXT("GenerateAuthToken");
+	static constexpr TCHAR Name[] = TEXT("ModifyAccountAttributes");
 
 	struct Params
 	{
+		/** The online account id of the Local User making the request. */
 		FOnlineAccountIdHandle LocalUserId;
-		FString Type;
-		TArray<FString> Scopes;
+		/** New or changed attributes. */
+		TMap<FSchemaAttributeId, FSchemaVariant> MutatedAttributes;
+		/** Attributes to be cleared. */
+		TSet<FSchemaAttributeId> ClearedAttributes;
 	};
 
 	struct Result
 	{
-		FString Type;
-		FCredentialsToken Token;
 	};
 };
 
-struct FAuthGetAuthToken
+struct FAuthQueryExternalServerAuthTicket
 {
-	static constexpr TCHAR Name[] = TEXT("GetAuthToken");
+	static constexpr TCHAR Name[] = TEXT("QueryExternalServerAuthTicket");
 
 	struct Params
 	{
+		/** The online account id of the Local User making the request. */
 		FOnlineAccountIdHandle LocalUserId;
-		FString Type;
-		TArray<FString> Scopes;
 	};
 
 	struct Result
 	{
-		FString Token;
+		FExternalServerAuthTicket ExternalServerAuthTicket;
 	};
 };
 
-struct FAuthGenerateAuthCode
+struct FAuthQueryExternalAuthToken
 {
-	static constexpr TCHAR Name[] = TEXT("GenerateAuthCode");
+	static constexpr TCHAR Name[] = TEXT("QueryExternalAuthToken");
 
 	struct Params
 	{
+		/** The online account id of the Local User making the request. */
 		FOnlineAccountIdHandle LocalUserId;
-		FString Type;
-		TArray<FString> Scopes;
+		/** The method of external auth to provide. */
+		EExternalAuthTokenMethod Method = EExternalAuthTokenMethod::Primary;
 	};
 
 	struct Result
 	{
-		FString Code;
+		FExternalAuthToken ExternalAuthToken;
 	};
 };
 
-struct FAuthGetAccountByPlatformUserId
+struct FAuthQueryVerifiedAuthTicket
 {
-	static constexpr TCHAR Name[] = TEXT("GetAccountByPlatformUserId");
+	static constexpr TCHAR Name[] = TEXT("QueryVerifiedAuthTicket");
 
 	struct Params
 	{
-		FPlatformUserId PlatformUserId = PLATFORMUSERID_NONE;
+		/** The online account id of the Local User making the request. */
+		FOnlineAccountIdHandle LocalUserId;
+		/** The intended purpose of the auth ticket. */
+		ERemoteAuthTicketAudience Audience = ERemoteAuthTicketAudience::Peer;
 	};
 
 	struct Result
 	{
-		TSharedRef<FAccountInfo> AccountInfo;
+		/** Local ticket id used to reference the ticket in further operations. */
+		FOnlineVerifiedAuthTicketIdHandle VerifiedAuthTicketId;
+		/** Ticket used to begin a verified auth session with a remote host. */
+		FVerifiedAuthTicket VerifiedAuthTicket;
 	};
 };
 
-struct FAuthGetAccountByAccountId
+struct FAuthCancelVerifiedAuthTicket
 {
-	static constexpr TCHAR Name[] = TEXT("GetAccountByAccountId");
+	static constexpr TCHAR Name[] = TEXT("CancelVerifiedAuthTicket");
 
 	struct Params
 	{
+		/** The online account id of the Local User making the request. */
+		FOnlineAccountIdHandle LocalUserId;
+		/** Local ticket id used to reference the verified auth ticket. */
+		FOnlineVerifiedAuthTicketIdHandle VerifiedAuthTicketId;
+	};
+
+	struct Result
+	{
+	};
+};
+
+struct FAuthBeginVerifiedAuthSession
+{
+	static constexpr TCHAR Name[] = TEXT("BeginVerifiedAuthSession");
+
+	struct Params
+	{
+		/** The remote user for which to start the verified auth session. */
+		FOnlineAccountIdHandle RemoteAccountId;
+		/** The ticket used to verify the identity of the user. */
+		FVerifiedAuthTicket Ticket;
+	};
+
+	struct Result
+	{
+		FVerifiedAuthSession Session;
+	};
+};
+
+struct FAuthEndVerifiedAuthSession
+{
+	static constexpr TCHAR Name[] = TEXT("EndVerifiedAuthSession");
+
+	struct Params
+	{
+		FOnlineVerifiedAuthSessionIdHandle SessionId;
+	};
+
+	struct Result
+	{
+	};
+};
+
+struct FAuthGetLocalOnlineUserByOnlineAccountId
+{
+	static constexpr TCHAR Name[] = TEXT("GetLocalOnlineUserByOnlineAccountId");
+
+	struct Params
+	{
+		/** Account Id of the Local User making the request */
 		FOnlineAccountIdHandle LocalUserId;
 	};
 
 	struct Result
 	{
 		TSharedRef<FAccountInfo> AccountInfo;
+
+	public:
+		Result() = delete; // cannot default construct due to TSharedRef
+	};
+};
+
+struct FAuthGetLocalOnlineUserByPlatformUserId
+{
+	static constexpr TCHAR Name[] = TEXT("GetLocalOnlineUserByPlatformUserId");
+
+	struct Params
+	{
+		/** The PlatformUserId of the Local User making the request. */
+		FPlatformUserId PlatformUserId;
+	};
+
+	struct Result
+	{
+		TSharedRef<FAccountInfo> AccountInfo;
+
+	public:
+		Result() = delete; // cannot default construct due to TSharedRef
+	};
+};
+
+struct FAuthGetAllLocalOnlineUsers
+{
+	static constexpr TCHAR Name[] = TEXT("GetAllLocalOnlineUsers");
+
+	struct Params
+	{
+	};
+
+	struct Result
+	{
+		TArray<TSharedRef<FAccountInfo>> AccountInfo;
 	};
 };
 
 /** Struct for LoginStatusChanged event */
-struct FLoginStatusChanged
+struct FAuthLoginStatusChanged
 {
-	/** User id whose status has changed */
-	FOnlineAccountIdHandle LocalUserId;
-	/** Previous login status */
-	ELoginStatus PreviousStatus;
-	/** Current login status */
-	ELoginStatus CurrentStatus;
+	/* The affected account. */
+	TSharedRef<FAccountInfo> AccountInfo;
+	/* The new login status. */
+	ELoginStatus LoginStatus = ELoginStatus::NotLoggedIn;
+
+public:
+	FAuthLoginStatusChanged() = delete; // cannot default construct due to TSharedRef
+};
+
+/** Struct for PendingAuthExpiration event */
+struct FAuthPendingAuthExpiration
+{
+	/* The affected account. */
+	TSharedRef<FAccountInfo> AccountInfo;
+
+public:
+	FAuthPendingAuthExpiration() = delete; // cannot default construct due to TSharedRef
+};
+
+/** Struct for AccountAttributesChanged event */
+struct FAuthAccountAttributesChanged
+{
+	/* The affected account. */
+	TSharedRef<FAccountInfo> AccountInfo;
+	/* Added attributes and their values. */
+	TMap<FSchemaAttributeId, FSchemaVariant> AddedAttributes;
+	/* Removed attributes with their previous values. */
+	TMap<FSchemaAttributeId, FSchemaVariant> RemovedAttributes;
+	/* Changed attributes with their old and new values. */
+	TMap<FSchemaAttributeId, TPair<FSchemaVariant, FSchemaVariant>> ChangedAttributes;
+
+public:
+	FAuthAccountAttributesChanged() = delete; // cannot default construct due to TSharedRef
 };
 
 class IAuth
 {
 public:
 	/**
-	 *
+	 * Authenticate a local user. If necessary, retrieve and maintain an access token for the
+	 * duration of the auth session.
 	 */
 	virtual TOnlineAsyncOpHandle<FAuthLogin> Login(FAuthLogin::Params&& Params) = 0;
 
 	/**
-	 *
+	 * Concludes the auth session for the local user.
 	 */
 	virtual TOnlineAsyncOpHandle<FAuthLogout> Logout(FAuthLogout::Params&& Params) = 0;
 
 	/**
+	 * Modify attributes associated with an authenticated account. Signals OnAccountAttributesChanged on completion.
+	 */
+	virtual TOnlineAsyncOpHandle<FAuthModifyAccountAttributes> ModifyAccountAttributes(FAuthModifyAccountAttributes::Params&& Params) = 0;
+
+	/**
+	 * Queries a ticket which is appropriate for making server-to-server calls on behalf of the
+	 * signed in user. Tickets are intended to be single use - users must call the API again to
+	 * retrieve a new ticket when making repeated calls which use a ticket.
+	 */
+	virtual TOnlineAsyncOpHandle<FAuthQueryExternalServerAuthTicket> QueryExternalServerAuthTicket(FAuthQueryExternalServerAuthTicket::Params&& Params) = 0;
+
+	/**
+	 * Retrieves a token appropriate for linking the service account with a service account of a
+	 * different service type. On most platforms this will return an OpenId token. May return a
+	 * cached token, but will handle refreshing if necessary.
+	 */
+	virtual TOnlineAsyncOpHandle<FAuthQueryExternalAuthToken> QueryExternalAuthToken(FAuthQueryExternalAuthToken::Params&& Params) = 0;
+
+	/**
+	 * Retrieves a ticket which is used to create a verified authentication session on a remote
+	 * client. When establishing a verified authentication session the user must always call
+	 * QueryVerifiedAuthTicket to retrieve a new ticket.
+	 * 
+	 * To prevent the transmission of PII to a non-trusted destination an audience enumeration
+	 * is required which may generate an empty ticket in the case of some peer-to-peer
+	 * implementations. On those platforms user authentication is handled implicitly through the
+	 * peer-to-peer networking model.
 	 *
+	 * It is the responsibility of the game code to provide an audience enumeration value
+	 * appropriate for the usage.
 	 */
-	virtual TOnlineAsyncOpHandle<FAuthGenerateAuthToken> GenerateAuthToken(FAuthGenerateAuthToken::Params&& Params) = 0;
+	virtual TOnlineAsyncOpHandle<FAuthQueryVerifiedAuthTicket> QueryVerifiedAuthTicket(FAuthQueryVerifiedAuthTicket::Params&& Params) = 0;
 
 	/**
-	 * Retrive a cached auth token
+	 * Cancels the ticket associated with a verified auth session. When the session created using
+	 * a verified auth ticket is no longer in use, CancelVerifiedAuthTicket must be called to clean
+	 * up any resources associated with the ticket.
 	 */
-	virtual TOnlineResult<FAuthGetAuthToken> GetAuthToken(FAuthGetAuthToken::Params&& Params) = 0;
+	virtual TOnlineAsyncOpHandle<FAuthCancelVerifiedAuthTicket> CancelVerifiedAuthTicket(FAuthCancelVerifiedAuthTicket::Params&& Params) = 0;
 
 	/**
+	 * Starts a verified auth session for a remote user. Depending on implementation
+	 * RemoteAuthTicket will be used to verify the identity of RemoteUserId. The intended usage is
+	 * to start a verified auth session for any client connecting any other client or server. Having
+	 * a verified auth session handles identity verification in both peer-to-peer and client-server
+	 * topologies. In peer-to-peer topologies authentication may not occur due to implicit trust
+	 * established by the peer-to-peer network.
 	 *
+	 * Creating a verified auth session will generate a FVerifiedAuthSessionId. This id is a
+	 * cryptographically secure random id intended to be transmitted between the established peers
+	 * to refer to the remote auth session. This id is needed when the client changes how they
+	 * communicate with the server. An example is connecting to a server reservation beacon to
+	 * reserve a spot before traveling to the server url using a different connection.
 	 */
-	virtual TOnlineAsyncOpHandle<FAuthGenerateAuthCode> GenerateAuthCode(FAuthGenerateAuthCode::Params&& Params) = 0;
+	virtual TOnlineAsyncOpHandle<FAuthBeginVerifiedAuthSession> BeginVerifiedAuthSession(FAuthBeginVerifiedAuthSession::Params&& Params) = 0;
 
 	/**
-	 * Get logged in user by local user num
-	 */ 
-	virtual TOnlineResult<FAuthGetAccountByPlatformUserId> GetAccountByPlatformUserId(FAuthGetAccountByPlatformUserId::Params&& Params) = 0;
+	 * Clean up the remote verified auth session and handle any required book-keeping.
+	 */
+	virtual TOnlineAsyncOpHandle<FAuthEndVerifiedAuthSession> EndVerifiedAuthSession(FAuthEndVerifiedAuthSession::Params&& Params) = 0;
 
 	/**
-	 * Get logged in user by account id
+	 * Retrieve a logged in user account.
 	 */
-	virtual TOnlineResult<FAuthGetAccountByAccountId> GetAccountByAccountId(FAuthGetAccountByAccountId::Params&& Params) = 0;
+	virtual TOnlineResult<FAuthGetLocalOnlineUserByOnlineAccountId> GetLocalOnlineUserByOnlineAccountId(FAuthGetLocalOnlineUserByOnlineAccountId::Params&& Params) const = 0;
 
 	/**
-	 * Event triggered when a local user's login status changes
+	 * Retrieve a logged in user account.
 	 */
-	virtual TOnlineEvent<void(const FLoginStatusChanged&)> OnLoginStatusChanged() = 0;
-};
+	virtual TOnlineResult<FAuthGetLocalOnlineUserByPlatformUserId> GetLocalOnlineUserByPlatformUserId(FAuthGetLocalOnlineUserByPlatformUserId::Params&& Params) const = 0;
 
+	/**
+	 * Retrieve all logged in user accounts.
+	 */
+	virtual TOnlineResult<FAuthGetAllLocalOnlineUsers> GetAllLocalOnlineUsers(FAuthGetAllLocalOnlineUsers::Params&& Params) const = 0;
+
+	/**
+	 * Triggered when the login status for a logged in user changes.
+	 */
+	virtual TOnlineEvent<void(const FAuthLoginStatusChanged&)> OnLoginStatusChanged() = 0;
+
+	/**
+	 * Triggered when the auth token will expire soon. Services which allow login using an external
+	 * auth token may trigger this notification when an updated external auth token is required to
+	 * maintain the auth session without interruption.
+	 */
+	virtual TOnlineEvent<void(const FAuthPendingAuthExpiration&)> OnPendingAuthExpiration() = 0;
+
+	/**
+	 * Triggered when the additional attributes associated with an authenticated account are changed.
+	 */
+	virtual TOnlineEvent<void(const FAuthAccountAttributesChanged&)> OnAccountAttributesChanged() = 0;
+
+	/**
+	 * Helper for querying the login status of a local user.
+	 */
+	virtual bool IsLoggedIn(const FOnlineAccountIdHandle& AccountId) const = 0;
+};
 
 namespace Meta {
 // TODO: Move to Auth_Meta.inl file?
 
 BEGIN_ONLINE_STRUCT_META(FAccountInfo)
+	ONLINE_STRUCT_FIELD(FAccountInfo, AccountId),
 	ONLINE_STRUCT_FIELD(FAccountInfo, PlatformUserId),
-	ONLINE_STRUCT_FIELD(FAccountInfo, UserId),
-	ONLINE_STRUCT_FIELD(FAccountInfo, LoginStatus)
+	ONLINE_STRUCT_FIELD(FAccountInfo, LoginStatus),
+	ONLINE_STRUCT_FIELD(FAccountInfo, Attributes)
+END_ONLINE_STRUCT_META()
+
+BEGIN_ONLINE_STRUCT_META(FExternalAuthToken)
+	ONLINE_STRUCT_FIELD(FExternalAuthToken, Type),
+	ONLINE_STRUCT_FIELD(FExternalAuthToken, Data)
+END_ONLINE_STRUCT_META()
+
+BEGIN_ONLINE_STRUCT_META(FExternalServerAuthTicket)
+	ONLINE_STRUCT_FIELD(FExternalServerAuthTicket, Type),
+	ONLINE_STRUCT_FIELD(FExternalServerAuthTicket, Data)
+END_ONLINE_STRUCT_META()
+
+BEGIN_ONLINE_STRUCT_META(FVerifiedAuthSession)
+	ONLINE_STRUCT_FIELD(FVerifiedAuthSession, SessionId),
+	ONLINE_STRUCT_FIELD(FVerifiedAuthSession, RemoteAccountId),
+	ONLINE_STRUCT_FIELD(FVerifiedAuthSession, CreationTime)
+END_ONLINE_STRUCT_META()
+
+BEGIN_ONLINE_STRUCT_META(FVerifiedAuthTicket)
+	ONLINE_STRUCT_FIELD(FVerifiedAuthTicket, Type),
+	ONLINE_STRUCT_FIELD(FVerifiedAuthTicket, Data)
 END_ONLINE_STRUCT_META()
 
 BEGIN_ONLINE_STRUCT_META(FAuthLogin::Params)
@@ -235,63 +563,109 @@ BEGIN_ONLINE_STRUCT_META(FAuthLogin::Result)
 END_ONLINE_STRUCT_META()
 
 BEGIN_ONLINE_STRUCT_META(FAuthLogout::Params)
-	ONLINE_STRUCT_FIELD(FAuthLogout::Params, LocalUserId)
+	ONLINE_STRUCT_FIELD(FAuthLogout::Params, LocalUserId),
+	ONLINE_STRUCT_FIELD(FAuthLogout::Params, bDestroyAuth)
 END_ONLINE_STRUCT_META()
 
 BEGIN_ONLINE_STRUCT_META(FAuthLogout::Result)
 END_ONLINE_STRUCT_META()
 
-BEGIN_ONLINE_STRUCT_META(FAuthGenerateAuthToken::Params)
-	ONLINE_STRUCT_FIELD(FAuthGenerateAuthToken::Params, LocalUserId),
-	ONLINE_STRUCT_FIELD(FAuthGenerateAuthToken::Params, Type),
-	ONLINE_STRUCT_FIELD(FAuthGenerateAuthToken::Params, Scopes)
+BEGIN_ONLINE_STRUCT_META(FAuthModifyAccountAttributes::Params)
+	ONLINE_STRUCT_FIELD(FAuthModifyAccountAttributes::Params, LocalUserId),
+	ONLINE_STRUCT_FIELD(FAuthModifyAccountAttributes::Params, MutatedAttributes),
+	ONLINE_STRUCT_FIELD(FAuthModifyAccountAttributes::Params, ClearedAttributes)
+END_ONLINE_STRUCT_META()
+
+BEGIN_ONLINE_STRUCT_META(FAuthModifyAccountAttributes::Result)
+END_ONLINE_STRUCT_META()
+
+BEGIN_ONLINE_STRUCT_META(FAuthQueryExternalServerAuthTicket::Params)
+	ONLINE_STRUCT_FIELD(FAuthQueryExternalServerAuthTicket::Params, LocalUserId)
+END_ONLINE_STRUCT_META()
+
+BEGIN_ONLINE_STRUCT_META(FAuthQueryExternalServerAuthTicket::Result)
+	ONLINE_STRUCT_FIELD(FAuthQueryExternalServerAuthTicket::Result, ExternalServerAuthTicket)
+END_ONLINE_STRUCT_META()
+
+BEGIN_ONLINE_STRUCT_META(FAuthQueryExternalAuthToken::Params)
+	ONLINE_STRUCT_FIELD(FAuthQueryExternalAuthToken::Params, LocalUserId),
+	ONLINE_STRUCT_FIELD(FAuthQueryExternalAuthToken::Params, Method)
+END_ONLINE_STRUCT_META()
+
+BEGIN_ONLINE_STRUCT_META(FAuthQueryExternalAuthToken::Result)
+	ONLINE_STRUCT_FIELD(FAuthQueryExternalAuthToken::Result, ExternalAuthToken)
+END_ONLINE_STRUCT_META()
+
+BEGIN_ONLINE_STRUCT_META(FAuthQueryVerifiedAuthTicket::Params)
+	ONLINE_STRUCT_FIELD(FAuthQueryVerifiedAuthTicket::Params, LocalUserId),
+	ONLINE_STRUCT_FIELD(FAuthQueryVerifiedAuthTicket::Params, Audience)
+END_ONLINE_STRUCT_META()
+
+BEGIN_ONLINE_STRUCT_META(FAuthQueryVerifiedAuthTicket::Result)
+	ONLINE_STRUCT_FIELD(FAuthQueryVerifiedAuthTicket::Result, VerifiedAuthTicket)
+END_ONLINE_STRUCT_META()
+
+BEGIN_ONLINE_STRUCT_META(FAuthCancelVerifiedAuthTicket::Params)
+	ONLINE_STRUCT_FIELD(FAuthCancelVerifiedAuthTicket::Params, LocalUserId),
+	ONLINE_STRUCT_FIELD(FAuthCancelVerifiedAuthTicket::Params, VerifiedAuthTicketId)
+END_ONLINE_STRUCT_META()
+
+BEGIN_ONLINE_STRUCT_META(FAuthCancelVerifiedAuthTicket::Result)
+END_ONLINE_STRUCT_META()
+
+BEGIN_ONLINE_STRUCT_META(FAuthBeginVerifiedAuthSession::Params)
+	ONLINE_STRUCT_FIELD(FAuthBeginVerifiedAuthSession::Params, RemoteAccountId),
+	ONLINE_STRUCT_FIELD(FAuthBeginVerifiedAuthSession::Params, Ticket)
+END_ONLINE_STRUCT_META()
+
+BEGIN_ONLINE_STRUCT_META(FAuthBeginVerifiedAuthSession::Result)
+	ONLINE_STRUCT_FIELD(FAuthBeginVerifiedAuthSession::Result, Session)
+END_ONLINE_STRUCT_META()
+
+BEGIN_ONLINE_STRUCT_META(FAuthEndVerifiedAuthSession::Params)
+	ONLINE_STRUCT_FIELD(FAuthEndVerifiedAuthSession::Params, SessionId)
 END_ONLINE_STRUCT_META()
 
-BEGIN_ONLINE_STRUCT_META(FAuthGenerateAuthToken::Result)
-	ONLINE_STRUCT_FIELD(FAuthGenerateAuthToken::Result, Type),
-	ONLINE_STRUCT_FIELD(FAuthGenerateAuthToken::Result, Token)
+BEGIN_ONLINE_STRUCT_META(FAuthEndVerifiedAuthSession::Result)
 END_ONLINE_STRUCT_META()
 
-BEGIN_ONLINE_STRUCT_META(FAuthGetAuthToken::Params)
-	ONLINE_STRUCT_FIELD(FAuthGetAuthToken::Params, LocalUserId),
-	ONLINE_STRUCT_FIELD(FAuthGetAuthToken::Params, Type),
-	ONLINE_STRUCT_FIELD(FAuthGetAuthToken::Params, Scopes)
+BEGIN_ONLINE_STRUCT_META(FAuthGetLocalOnlineUserByOnlineAccountId::Params)
+	ONLINE_STRUCT_FIELD(FAuthGetLocalOnlineUserByOnlineAccountId::Params, LocalUserId)
 END_ONLINE_STRUCT_META()
 
-BEGIN_ONLINE_STRUCT_META(FAuthGetAuthToken::Result)
-	ONLINE_STRUCT_FIELD(FAuthGetAuthToken::Result, Token)
+BEGIN_ONLINE_STRUCT_META(FAuthGetLocalOnlineUserByOnlineAccountId::Result)
+	ONLINE_STRUCT_FIELD(FAuthGetLocalOnlineUserByOnlineAccountId::Result, AccountInfo)
 END_ONLINE_STRUCT_META()
 
-BEGIN_ONLINE_STRUCT_META(FAuthGenerateAuthCode::Params)
-ONLINE_STRUCT_FIELD(FAuthGenerateAuthCode::Params, LocalUserId),
-ONLINE_STRUCT_FIELD(FAuthGenerateAuthCode::Params, Type),
-ONLINE_STRUCT_FIELD(FAuthGenerateAuthCode::Params, Scopes)
+BEGIN_ONLINE_STRUCT_META(FAuthGetLocalOnlineUserByPlatformUserId::Params)
+	ONLINE_STRUCT_FIELD(FAuthGetLocalOnlineUserByPlatformUserId::Params, PlatformUserId)
 END_ONLINE_STRUCT_META()
 
-BEGIN_ONLINE_STRUCT_META(FAuthGenerateAuthCode::Result)
-	ONLINE_STRUCT_FIELD(FAuthGenerateAuthCode::Result, Code)
+BEGIN_ONLINE_STRUCT_META(FAuthGetLocalOnlineUserByPlatformUserId::Result)
+	ONLINE_STRUCT_FIELD(FAuthGetLocalOnlineUserByPlatformUserId::Result, AccountInfo)
 END_ONLINE_STRUCT_META()
 
-BEGIN_ONLINE_STRUCT_META(FAuthGetAccountByPlatformUserId::Params)
-	ONLINE_STRUCT_FIELD(FAuthGetAccountByPlatformUserId::Params, PlatformUserId)
+BEGIN_ONLINE_STRUCT_META(FAuthGetAllLocalOnlineUsers::Params)
 END_ONLINE_STRUCT_META()
 
-BEGIN_ONLINE_STRUCT_META(FAuthGetAccountByPlatformUserId::Result)
-	ONLINE_STRUCT_FIELD(FAuthGetAccountByPlatformUserId::Result, AccountInfo)
+BEGIN_ONLINE_STRUCT_META(FAuthGetAllLocalOnlineUsers::Result)
+	ONLINE_STRUCT_FIELD(FAuthGetAllLocalOnlineUsers::Result, AccountInfo)
 END_ONLINE_STRUCT_META()
 
-BEGIN_ONLINE_STRUCT_META(FAuthGetAccountByAccountId::Params)
-	ONLINE_STRUCT_FIELD(FAuthGetAccountByAccountId::Params, LocalUserId)
+BEGIN_ONLINE_STRUCT_META(FAuthLoginStatusChanged)
+	ONLINE_STRUCT_FIELD(FAuthLoginStatusChanged, AccountInfo),
+	ONLINE_STRUCT_FIELD(FAuthLoginStatusChanged, LoginStatus)
 END_ONLINE_STRUCT_META()
 
-BEGIN_ONLINE_STRUCT_META(FAuthGetAccountByAccountId::Result)
-	ONLINE_STRUCT_FIELD(FAuthGetAccountByAccountId::Result, AccountInfo)
+BEGIN_ONLINE_STRUCT_META(FAuthPendingAuthExpiration)
+	ONLINE_STRUCT_FIELD(FAuthPendingAuthExpiration, AccountInfo)
 END_ONLINE_STRUCT_META()
 
-BEGIN_ONLINE_STRUCT_META(FLoginStatusChanged)
-	ONLINE_STRUCT_FIELD(FLoginStatusChanged, LocalUserId),
-	ONLINE_STRUCT_FIELD(FLoginStatusChanged, PreviousStatus),
-	ONLINE_STRUCT_FIELD(FLoginStatusChanged, CurrentStatus)
+BEGIN_ONLINE_STRUCT_META(FAuthAccountAttributesChanged)
+	ONLINE_STRUCT_FIELD(FAuthAccountAttributesChanged, AccountInfo),
+	ONLINE_STRUCT_FIELD(FAuthAccountAttributesChanged, AddedAttributes),
+	ONLINE_STRUCT_FIELD(FAuthAccountAttributesChanged, RemovedAttributes),
+	ONLINE_STRUCT_FIELD(FAuthAccountAttributesChanged, ChangedAttributes)
 END_ONLINE_STRUCT_META()
 
 /* Meta*/ }
diff --git a/Engine/Plugins/Online/Experimental/OnlineServices/Source/OnlineServicesInterface/Public/Online/Lobbies.h b/Engine/Plugins/Online/Experimental/OnlineServices/Source/OnlineServicesInterface/Public/Online/Lobbies.h
--- a/Engine/Plugins/Online/Experimental/OnlineServices/Source/OnlineServicesInterface/Public/Online/Lobbies.h
+++ b/Engine/Plugins/Online/Experimental/OnlineServices/Source/OnlineServicesInterface/Public/Online/Lobbies.h
@@ -5,43 +5,14 @@
 #include "Online/OnlineAsyncOpHandle.h"
 #include "Online/CoreOnline.h"
 #include "Online/OnlineMeta.h"
+#include "Online/Schema.h"
 #include "Misc/TVariant.h"
 
 namespace UE::Online {
 
-using FLobbySchemaId = FName;
-using FLobbyAttributeId = FName;
-
-class FLobbyVariant
-{
-public:
-	using FVariantType = TVariant<FString, int64, double, bool>;
-	FLobbyVariant() = default;
-	FLobbyVariant(const FLobbyVariant& InOther) : VariantData(InOther.VariantData) {}
-	FLobbyVariant(FLobbyVariant&& InOther) : VariantData(MoveTemp(InOther.VariantData)) {}
-	FLobbyVariant& operator=(FLobbyVariant&&) = default;
-	FLobbyVariant& operator=(const FLobbyVariant&) = default;
-
-	template<typename ValueType>
-	FLobbyVariant(const ValueType& InData) { Set(MoveTempIfPossible(InData)); }
-	void Set(const TCHAR* AsString) { VariantData.Emplace<FString>(AsString); }
-	void Set(const FString& AsString) { VariantData.Emplace<FString>(AsString); }
-	void Set(FString&& AsString) { VariantData.Emplace<FString>(MoveTemp(AsString)); }
-	void Set(int64 AsInt) { VariantData.Emplace<int64>(AsInt); }
-	void Set(double AsDouble) { VariantData.Emplace<double>(AsDouble); }
-	void Set(bool bAsBool) { VariantData.Emplace<bool>(bAsBool); }
-	ONLINESERVICESINTERFACE_API int64 GetInt64() const;
-	ONLINESERVICESINTERFACE_API double GetDouble() const;
-	ONLINESERVICESINTERFACE_API bool GetBoolean() const;
-	ONLINESERVICESINTERFACE_API FString GetString() const;
-
-	ONLINESERVICESINTERFACE_API bool operator==(const FLobbyVariant& Other) const;
-	bool operator!=(const FLobbyVariant& Other) const { return !(*this == Other); }
-public:
-	FVariantType VariantData;
-};
-ONLINESERVICESINTERFACE_API const FString LexToString(const FLobbyVariant& Variant);
-ONLINESERVICESINTERFACE_API void LexFromString(FLobbyVariant& Variant, const TCHAR* InStr);
+using FLobbySchemaId = FSchemaId;
+using FLobbyAttributeId = FSchemaAttributeId;
+using FLobbyVariant = FSchemaVariant;
 
 enum class ELobbyJoinPolicy : uint8
 {
diff --git a/Engine/Plugins/Online/Experimental/OnlineServices/Source/OnlineServicesInterface/Public/Online/Schema.h b/Engine/Plugins/Online/Experimental/OnlineServices/Source/OnlineServicesInterface/Public/Online/Schema.h
new file mode 100644
--- /dev/null
+++ b/Engine/Plugins/Online/Experimental/OnlineServices/Source/OnlineServicesInterface/Public/Online/Schema.h
@@ -0,0 +1,46 @@
+// Copyright Epic Games, Inc. All Rights Reserved.
+
+#pragma once
+
+#include "CoreMinimal.h"
+#include "Misc/TVariant.h"
+
+namespace UE::Online {
+
+using FSchemaId = FName;
+using FSchemaAttributeId = FName;
+
+class ONLINESERVICESINTERFACE_API FSchemaVariant final
+{
+public:
+	using FVariantType = TVariant<FString, int64, double, bool>;
+	FSchemaVariant() = default;
+	FSchemaVariant(const FSchemaVariant& InOther) = default;
+	FSchemaVariant(FSchemaVariant&& InOther) = default;
+	FSchemaVariant& operator=(FSchemaVariant&&) = default;
+	FSchemaVariant& operator=(const FSchemaVariant&) = default;
+
+	template<typename ValueType>
+	FSchemaVariant(const ValueType& InData) { Set(InData); }
+	template<typename ValueType>
+	FSchemaVariant(ValueType&& InData) { Set(MoveTemp(InData)); }
+	void Set(const TCHAR* AsString) { VariantData.Emplace<FString>(AsString); }
+	void Set(const FString& AsString) { VariantData.Emplace<FString>(AsString); }
+	void Set(FString&& AsString) { VariantData.Emplace<FString>(MoveTemp(AsString)); }
+	void Set(int64 AsInt) { VariantData.Emplace<int64>(AsInt); }
+	void Set(double AsDouble) { VariantData.Emplace<double>(AsDouble); }
+	void Set(bool bAsBool) { VariantData.Emplace<bool>(bAsBool); }
+	int64 GetInt64() const;
+	double GetDouble() const;
+	bool GetBoolean() const;
+	FString GetString() const;
+
+	bool operator==(const FSchemaVariant& Other) const;
+	bool operator!=(const FSchemaVariant& Other) const { return !(*this == Other); }
+public:
+	FVariantType VariantData;
+};
+ONLINESERVICESINTERFACE_API const FString LexToString(const FSchemaVariant& Variant);
+ONLINESERVICESINTERFACE_API void LexFromString(FSchemaVariant& Variant, const TCHAR* InStr);
+
+/* UE::Online */ }
diff --git a/Engine/Source/Runtime/CoreOnline/Private/Online/CoreOnline.cpp b/Engine/Source/Runtime/CoreOnline/Private/Online/CoreOnline.cpp
--- a/Engine/Source/Runtime/CoreOnline/Private/Online/CoreOnline.cpp
+++ b/Engine/Source/Runtime/CoreOnline/Private/Online/CoreOnline.cpp
@@ -305,6 +305,16 @@
 	return FString::Printf(TEXT("%s:%d"), LexToString(Id.GetOnlineServicesType()), Id.GetHandle());
 }
 
+FString ToLogString(const FOnlineVerifiedAuthTicketIdHandle& Id)
+{
+	return FString::Printf(TEXT("%s:%d"), LexToString(Id.GetOnlineServicesType()), Id.GetHandle());
+}
+
+FString ToLogString(const FOnlineVerifiedAuthSessionIdHandle& Id)
+{
+	return FString::Printf(TEXT("%s:%d"), LexToString(Id.GetOnlineServicesType()), Id.GetHandle());
+}
+
 }	/* UE::Online */
 
 FString FUniqueNetIdWrapper::ToDebugString() const
diff --git a/Engine/Source/Runtime/CoreOnline/Public/Online/CoreOnline.h b/Engine/Source/Runtime/CoreOnline/Public/Online/CoreOnline.h
--- a/Engine/Source/Runtime/CoreOnline/Public/Online/CoreOnline.h
+++ b/Engine/Source/Runtime/CoreOnline/Public/Online/CoreOnline.h
@@ -198,6 +198,8 @@
 	struct FAccount {};
 	struct FSession {};
 	struct FLobby {};
+	struct FVerifiedAuthTicket {};
+	struct FVerifiedAuthSession {};
 }
 
 enum class EOnlineServices : uint8
@@ -292,10 +294,14 @@
 using FOnlineAccountIdHandle = TOnlineIdHandle<OnlineIdHandleTags::FAccount>;
 using FOnlineLobbyIdHandle = TOnlineIdHandle<OnlineIdHandleTags::FLobby>;
 using FOnlineSessionIdHandle = TOnlineIdHandle<OnlineIdHandleTags::FSession>;
+using FOnlineVerifiedAuthTicketIdHandle = TOnlineIdHandle<OnlineIdHandleTags::FVerifiedAuthTicket>;
+using FOnlineVerifiedAuthSessionIdHandle = TOnlineIdHandle<OnlineIdHandleTags::FVerifiedAuthSession>;
 
 COREONLINE_API FString ToLogString(const FOnlineAccountIdHandle& Id);
 COREONLINE_API FString ToLogString(const FOnlineLobbyIdHandle& Id);
 COREONLINE_API FString ToLogString(const FOnlineSessionIdHandle& Id);
+COREONLINE_API FString ToLogString(const FOnlineVerifiedAuthTicketIdHandle& Id);
+COREONLINE_API FString ToLogString(const FOnlineVerifiedAuthSessionIdHandle& Id);
 
 template<typename IdType>
 inline uint32 GetTypeHash(const TOnlineIdHandle<IdType>& Handle)
