OSSv2 Auth cleanup

diff --git a/Engine/Plugins/Online/Experimental/OnlineServices/Source/OnlineServicesCommon/Private/Online/LobbiesCommonTests.cpp b/Engine/Plugins/Online/Experimental/OnlineServices/Source/OnlineServicesCommon/Private/Online/LobbiesCommonTests.cpp
--- a/Engine/Plugins/Online/Experimental/OnlineServices/Source/OnlineServicesCommon/Private/Online/LobbiesCommonTests.cpp
+++ b/Engine/Plugins/Online/Experimental/OnlineServices/Source/OnlineServicesCommon/Private/Online/LobbiesCommonTests.cpp
@@ -1187,7 +1187,7 @@ struct FFunctionalTestLoginUser
 	{
 		IAuth* AuthInterface = nullptr;
 		FPlatformUserId PlatformUserId;
-		FString Type;
+		FName Type;
 		FString Id;
 		FString Token;
 	};
@@ -1266,9 +1266,9 @@ TFuture<TOnlineResult<FFunctionalTestLogoutUser>> FunctionalTestLogoutUser(FFunc
 		return MakeFulfilledPromise<TOnlineResult<FFunctionalTestLogoutUser>>(Errors::MissingInterface()).GetFuture();
 	}
 
-	FAuthGetAccountByPlatformUserId::Params GetAccountParams;
+	FAuthGetLocalOnlineUserByPlatformUserId::Params GetAccountParams;
 	GetAccountParams.PlatformUserId = Params.PlatformUserId;
-	TOnlineResult<FAuthGetAccountByPlatformUserId> Result = Params.AuthInterface->GetAccountByPlatformUserId(MoveTemp(GetAccountParams));
+	TOnlineResult<FAuthGetLocalOnlineUserByPlatformUserId> Result = Params.AuthInterface->GetLocalOnlineUserByPlatformUserId(MoveTemp(GetAccountParams));
 	if (Result.IsError())
 	{
 		// Ignore errors for now. This function should ignore logged out users.
@@ -1292,7 +1292,7 @@ TFuture<TOnlineResult<FFunctionalTestLogoutUser>> FunctionalTestLogoutUser(FFunc
 	TFuture<TOnlineResult<FFunctionalTestLogoutUser>> Future = Promise->GetFuture();
 
 	FAuthLogout::Params LogoutParams;
-	LogoutParams.LocalUserId = AccountInfo->UserId;
+	LogoutParams.LocalUserId = AccountInfo->AccountId;
 
 	Params.AuthInterface->Logout(MoveTemp(LogoutParams))
 	.OnComplete([Promise, PlatformUserId = Params.PlatformUserId](const TOnlineResult<FAuthLogout>& LogoutResult)
@@ -1395,14 +1395,6 @@ private:
 
 } // Private
 
-template <typename DataType, typename OpType>
-const DataType& GetOpDataChecked(const TOnlineAsyncOp<OpType>& Op, const FString& Key)
-{
-	const DataType* Data = Op.Data.template Get<DataType>(Key);
-	check(Data);
-	return *Data;
-}
-
 struct FFunctionalTestConfig
 {
 	FString TestAccount1Type;
@@ -1456,8 +1448,8 @@ TOnlineAsyncOpHandle<FFunctionalTestLobbies> RunLobbyFunctionalTest(IAuth& AuthI
 	Op->Data.Set(LobbyEventCaptureKeyName, MakeShared<FLobbyEventCapture>(LobbyEvents));
 
 	Op->Then(Private::ConsumeStepResult<Private::FFunctionalTestLogoutAllUsers>(*Op, &Private::FunctionalTestLogoutAllUsers, Private::FFunctionalTestLogoutAllUsers::Params{&AuthInterface}))
-	.Then(Private::CaptureStepResult<Private::FFunctionalTestLoginUser>(*Op, User1KeyName, &Private::FunctionalTestLoginUser, Private::FFunctionalTestLoginUser::Params{&AuthInterface, FPlatformMisc::GetPlatformUserForUserIndex(0), TestConfig->TestAccount1Type, TestConfig->TestAccount1Id, TestConfig->TestAccount1Token}))
-	.Then(Private::CaptureStepResult<Private::FFunctionalTestLoginUser>(*Op, User2KeyName, &Private::FunctionalTestLoginUser, Private::FFunctionalTestLoginUser::Params{&AuthInterface, FPlatformMisc::GetPlatformUserForUserIndex(1), TestConfig->TestAccount2Type, TestConfig->TestAccount2Id, TestConfig->TestAccount2Token}))
+	.Then(Private::CaptureStepResult<Private::FFunctionalTestLoginUser>(*Op, User1KeyName, &Private::FunctionalTestLoginUser, Private::FFunctionalTestLoginUser::Params{&AuthInterface, FPlatformMisc::GetPlatformUserForUserIndex(0), *TestConfig->TestAccount1Type, TestConfig->TestAccount1Id, TestConfig->TestAccount1Token}))
+	.Then(Private::CaptureStepResult<Private::FFunctionalTestLoginUser>(*Op, User2KeyName, &Private::FunctionalTestLoginUser, Private::FFunctionalTestLoginUser::Params{&AuthInterface, FPlatformMisc::GetPlatformUserForUserIndex(1), *TestConfig->TestAccount2Type, TestConfig->TestAccount2Id, TestConfig->TestAccount2Token}))
 
 	//----------------------------------------------------------------------------------------------
 	// Test 1:
@@ -1476,13 +1468,13 @@ TOnlineAsyncOpHandle<FFunctionalTestLobbies> RunLobbyFunctionalTest(IAuth& AuthI
 
 		// Create a lobby
 		FCreateLobby::Params Params;
-		Params.LocalUserId = User1Info.AccountInfo->UserId;
+		Params.LocalUserId = User1Info.AccountInfo->AccountId;
 		Params.LocalName = TEXT("test");
 		Params.SchemaName = TEXT("test");
 		Params.MaxMembers = 2;
 		Params.JoinPolicy = ELobbyJoinPolicy::InvitationOnly;
 		Params.Attributes.Add(TEXT("test_attribute_key"), TEXT("test_attribute_value"));
-		Params.LocalUsers.Emplace(FJoinLobbyLocalUserData{User1Info.AccountInfo->UserId, {{TEXT("test_attribute_key"), TEXT("test_attribute_value")}}});
+		Params.LocalUsers.Emplace(FJoinLobbyLocalUserData{User1Info.AccountInfo->AccountId, {{TEXT("test_attribute_key"), TEXT("test_attribute_value")}}});
 		return Params;
 	})
 	.Then(Private::CaptureOperationStepResult<FCreateLobby>(*Op, CreateLobbyKeyName, Private::FBindOperation<FCreateLobby>(&LobbiesCommon, &FLobbiesCommon::CreateLobby)))
@@ -1507,7 +1499,7 @@ TOnlineAsyncOpHandle<FFunctionalTestLobbies> RunLobbyFunctionalTest(IAuth& AuthI
 			InAsyncOp.SetError(Errors::Cancelled());
 		}
 
-		const TSharedRef<const FLobbyMember>* MemberData = CreateResult.Lobby->Members.Find(User1Info.AccountInfo->UserId);
+		const TSharedRef<const FLobbyMember>* MemberData = CreateResult.Lobby->Members.Find(User1Info.AccountInfo->AccountId);
 		if (CreateResult.Lobby->Members.Num() == 1 && MemberData != nullptr)
 		{
 			if (!(**MemberData).Attributes.OrderIndependentCompareEqual(
@@ -1530,7 +1522,7 @@ TOnlineAsyncOpHandle<FFunctionalTestLobbies> RunLobbyFunctionalTest(IAuth& AuthI
 		const FCreateLobby::Result& CreateLobbyResult = GetOpDataChecked<FCreateLobby::Result>(InAsyncOp, CreateLobbyKeyName);
 
 		FModifyLobbyAttributes::Params Params;
-		Params.LocalUserId = User1Info.AccountInfo->UserId;
+		Params.LocalUserId = User1Info.AccountInfo->AccountId;
 		Params.LobbyId = CreateLobbyResult.Lobby->LobbyId;
 		Params.MutatedAttributes = {{TEXT("test_attribute_key"), TEXT("mutated_test_attribute_value")}};
 		return Params;
@@ -1563,7 +1555,7 @@ TOnlineAsyncOpHandle<FFunctionalTestLobbies> RunLobbyFunctionalTest(IAuth& AuthI
 		const FCreateLobby::Result& CreateLobbyResult = GetOpDataChecked<FCreateLobby::Result>(InAsyncOp, CreateLobbyKeyName);
 
 		FModifyLobbyMemberAttributes::Params Params;
-		Params.LocalUserId = User1Info.AccountInfo->UserId;
+		Params.LocalUserId = User1Info.AccountInfo->AccountId;
 		Params.LobbyId = CreateLobbyResult.Lobby->LobbyId;
 		Params.MutatedAttributes = {{TEXT("test_attribute_key"), TEXT("mutated_test_attribute_value")}};
 		return Params;
@@ -1578,7 +1570,7 @@ TOnlineAsyncOpHandle<FFunctionalTestLobbies> RunLobbyFunctionalTest(IAuth& AuthI
 		if (EventCapture->LobbyMemberAttributesChanged.Num() == 1)
 		{
 			const FLobbyMemberAttributesChanged& Notification = EventCapture->LobbyMemberAttributesChanged[0].Notification;
-			if (const TSharedRef<const FLobbyMember>* MemberData = Notification.Lobby->Members.Find(User1Info.AccountInfo->UserId))
+			if (const TSharedRef<const FLobbyMember>* MemberData = Notification.Lobby->Members.Find(User1Info.AccountInfo->AccountId))
 			{
 				if (!Notification.ChangedAttributes.Difference(TSet<FLobbyAttributeId>{TEXT("test_attribute_key")}).IsEmpty() ||
 					!(**MemberData).Attributes.OrderIndependentCompareEqual(TMap<FLobbyAttributeId, FLobbyVariant>{{TEXT("test_attribute_key"), TEXT("mutated_test_attribute_value")}}))
@@ -1605,7 +1597,7 @@ TOnlineAsyncOpHandle<FFunctionalTestLobbies> RunLobbyFunctionalTest(IAuth& AuthI
 		const FCreateLobby::Result& CreateLobbyResult = GetOpDataChecked<FCreateLobby::Result>(InAsyncOp, CreateLobbyKeyName);
 
 		FLeaveLobby::Params Params;
-		Params.LocalUserId = User1Info.AccountInfo->UserId;
+		Params.LocalUserId = User1Info.AccountInfo->AccountId;
 		Params.LobbyId = CreateLobbyResult.Lobby->LobbyId;
 		return Params;
 	})
@@ -1639,13 +1631,13 @@ TOnlineAsyncOpHandle<FFunctionalTestLobbies> RunLobbyFunctionalTest(IAuth& AuthI
 		const Private::FFunctionalTestLoginUser::Result& User1Info = GetOpDataChecked<Private::FFunctionalTestLoginUser::Result>(InAsyncOp, User1KeyName);
 
 		FCreateLobby::Params Params;
-		Params.LocalUserId = User1Info.AccountInfo->UserId;
+		Params.LocalUserId = User1Info.AccountInfo->AccountId;
 		Params.LocalName = TEXT("test");
 		Params.SchemaName = TEXT("test");
 		Params.MaxMembers = 2;
 		Params.JoinPolicy = ELobbyJoinPolicy::InvitationOnly;
 		//Params.Attributes;
-		Params.LocalUsers.Emplace(FJoinLobbyLocalUserData{User1Info.AccountInfo->UserId, {}});
+		Params.LocalUsers.Emplace(FJoinLobbyLocalUserData{User1Info.AccountInfo->AccountId, {}});
 		return Params;
 	})
 	.Then(Private::CaptureOperationStepResult<FCreateLobby>(*Op, CreateLobbyKeyName, Private::FBindOperation<FCreateLobby>(&LobbiesCommon, &FLobbiesCommon::CreateLobby)))
@@ -1670,9 +1662,9 @@ TOnlineAsyncOpHandle<FFunctionalTestLobbies> RunLobbyFunctionalTest(IAuth& AuthI
 		const FCreateLobby::Result& CreateLobbyResult = GetOpDataChecked<FCreateLobby::Result>(InAsyncOp, CreateLobbyKeyName);
 
 		FInviteLobbyMember::Params Params;
-		Params.LocalUserId = User1Info.AccountInfo->UserId;
+		Params.LocalUserId = User1Info.AccountInfo->AccountId;
 		Params.LobbyId = CreateLobbyResult.Lobby->LobbyId;
-		Params.TargetUserId = User2Info.AccountInfo->UserId;
+		Params.TargetUserId = User2Info.AccountInfo->AccountId;
 		return Params;
 	})
 	.Then(Private::ConsumeOperationStepResult<FInviteLobbyMember>(*Op, Private::FBindOperation<FInviteLobbyMember>(&LobbiesCommon, &FLobbiesCommon::InviteLobbyMember)))
@@ -1680,7 +1672,7 @@ TOnlineAsyncOpHandle<FFunctionalTestLobbies> RunLobbyFunctionalTest(IAuth& AuthI
 	{
 		const Private::FFunctionalTestLoginUser::Result& User2Info = GetOpDataChecked<Private::FFunctionalTestLoginUser::Result>(InAsyncOp, User2KeyName);
 		const FCreateLobby::Result& CreateLobbyResult = GetOpDataChecked<FCreateLobby::Result>(InAsyncOp, CreateLobbyKeyName);
-		return Private::AwaitInvitation(*LobbyEventsPtr, User2Info.AccountInfo->UserId, CreateLobbyResult.Lobby->LobbyId, TestConfig->InvitationWaitSeconds);
+		return Private::AwaitInvitation(*LobbyEventsPtr, User2Info.AccountInfo->AccountId, CreateLobbyResult.Lobby->LobbyId, TestConfig->InvitationWaitSeconds);
 	})
 	.Then([](TOnlineAsyncOp<FFunctionalTestLobbies>& InAsyncOp, TDefaultErrorResultInternal<FOnlineLobbyIdHandle>&& Result)
 	{
@@ -1698,10 +1690,10 @@ TOnlineAsyncOpHandle<FFunctionalTestLobbies> RunLobbyFunctionalTest(IAuth& AuthI
 		const FCreateLobby::Result& CreateLobbyResult = GetOpDataChecked<FCreateLobby::Result>(InAsyncOp, CreateLobbyKeyName);
 
 		FJoinLobby::Params Params;
-		Params.LocalUserId = User2Info.AccountInfo->UserId;
+		Params.LocalUserId = User2Info.AccountInfo->AccountId;
 		Params.LocalName = TEXT("test");
 		Params.LobbyId = CreateLobbyResult.Lobby->LobbyId;
-		Params.LocalUsers.Emplace(FJoinLobbyLocalUserData{User2Info.AccountInfo->UserId, {}});
+		Params.LocalUsers.Emplace(FJoinLobbyLocalUserData{User2Info.AccountInfo->AccountId, {}});
 		return Params;
 	})
 	.Then(Private::ConsumeOperationStepResult<FJoinLobby>(*Op, Private::FBindOperation<FJoinLobby>(&LobbiesCommon, &FLobbiesCommon::JoinLobby)))
@@ -1723,7 +1715,7 @@ TOnlineAsyncOpHandle<FFunctionalTestLobbies> RunLobbyFunctionalTest(IAuth& AuthI
 		const FCreateLobby::Result& CreateLobbyResult = GetOpDataChecked<FCreateLobby::Result>(InAsyncOp, CreateLobbyKeyName);
 
 		FLeaveLobby::Params Params;
-		Params.LocalUserId = User2Info.AccountInfo->UserId;
+		Params.LocalUserId = User2Info.AccountInfo->AccountId;
 		Params.LobbyId = CreateLobbyResult.Lobby->LobbyId;
 		return Params;
 	})
@@ -1746,7 +1738,7 @@ TOnlineAsyncOpHandle<FFunctionalTestLobbies> RunLobbyFunctionalTest(IAuth& AuthI
 		const FCreateLobby::Result& CreateLobbyResult = GetOpDataChecked<FCreateLobby::Result>(InAsyncOp, CreateLobbyKeyName);
 
 		FLeaveLobby::Params Params;
-		Params.LocalUserId = User1Info.AccountInfo->UserId;
+		Params.LocalUserId = User1Info.AccountInfo->AccountId;
 		Params.LobbyId = CreateLobbyResult.Lobby->LobbyId;
 		return Params;
 	})
@@ -1788,13 +1780,13 @@ TOnlineAsyncOpHandle<FFunctionalTestLobbies> RunLobbyFunctionalTest(IAuth& AuthI
 		InAsyncOp.Data.Set(SearchKeyName, TSharedRef<FSearchParams>(SearchParams));
 
 		FCreateLobby::Params Params;
-		Params.LocalUserId = User1Info.AccountInfo->UserId;
+		Params.LocalUserId = User1Info.AccountInfo->AccountId;
 		Params.LocalName = TEXT("test");
 		Params.SchemaName = TEXT("test");
 		Params.MaxMembers = 2;
 		Params.JoinPolicy = ELobbyJoinPolicy::PublicAdvertised;
 		Params.Attributes = {{ TEXT("LobbyCreateTime"), SearchParams->LobbyCreateTime }};
-		Params.LocalUsers.Emplace(FJoinLobbyLocalUserData{User1Info.AccountInfo->UserId, {}});
+		Params.LocalUsers.Emplace(FJoinLobbyLocalUserData{User1Info.AccountInfo->AccountId, {}});
 		return Params;
 	})
 	.Then(Private::CaptureOperationStepResult<FCreateLobby>(*Op, CreateLobbyKeyName, Private::FBindOperation<FCreateLobby>(&LobbiesCommon, &FLobbiesCommon::CreateLobby)))
@@ -1823,7 +1815,7 @@ TOnlineAsyncOpHandle<FFunctionalTestLobbies> RunLobbyFunctionalTest(IAuth& AuthI
 
 		// Search for lobby from create. Searching by attribute will also limit results by bucket id.
 		FFindLobbies::Params Params;
-		Params.LocalUserId = User2Info.AccountInfo->UserId;
+		Params.LocalUserId = User2Info.AccountInfo->AccountId;
 		Params.Filters = {{TEXT("LobbyCreateTime"), ELobbyComparisonOp::Equals, SearchParams->LobbyCreateTime}};
 		return Params;
 	})
@@ -1844,10 +1836,10 @@ TOnlineAsyncOpHandle<FFunctionalTestLobbies> RunLobbyFunctionalTest(IAuth& AuthI
 		}
 
 		FJoinLobby::Params Params;
-		Params.LocalUserId = User2Info.AccountInfo->UserId;
+		Params.LocalUserId = User2Info.AccountInfo->AccountId;
 		Params.LocalName = TEXT("test");
 		Params.LobbyId = FindResults.Lobbies[0]->LobbyId;
-		Params.LocalUsers.Emplace(FJoinLobbyLocalUserData{User2Info.AccountInfo->UserId, {}});
+		Params.LocalUsers.Emplace(FJoinLobbyLocalUserData{User2Info.AccountInfo->AccountId, {}});
 		return Params;
 	})
 	.Then(Private::ConsumeOperationStepResult<FJoinLobby>(*Op, Private::FBindOperation<FJoinLobby>(&LobbiesCommon, &FLobbiesCommon::JoinLobby)))
@@ -1869,7 +1861,7 @@ TOnlineAsyncOpHandle<FFunctionalTestLobbies> RunLobbyFunctionalTest(IAuth& AuthI
 		const FCreateLobby::Result& CreateLobbyResult = GetOpDataChecked<FCreateLobby::Result>(InAsyncOp, CreateLobbyKeyName);
 
 		FLeaveLobby::Params Params;
-		Params.LocalUserId = User2Info.AccountInfo->UserId;
+		Params.LocalUserId = User2Info.AccountInfo->AccountId;
 		Params.LobbyId = CreateLobbyResult.Lobby->LobbyId;
 		return Params;
 	})
@@ -1892,7 +1884,7 @@ TOnlineAsyncOpHandle<FFunctionalTestLobbies> RunLobbyFunctionalTest(IAuth& AuthI
 		const FCreateLobby::Result& CreateLobbyResult = GetOpDataChecked<FCreateLobby::Result>(InAsyncOp, CreateLobbyKeyName);
 
 		FLeaveLobby::Params Params;
-		Params.LocalUserId = User1Info.AccountInfo->UserId;
+		Params.LocalUserId = User1Info.AccountInfo->AccountId;
 		Params.LobbyId = CreateLobbyResult.Lobby->LobbyId;
 		return Params;
 	})
diff --git a/Engine/Plugins/Online/Experimental/OnlineServices/Source/OnlineServicesCommon/Private/Online/OnlineSyntaxTests.cpp b/Engine/Plugins/Online/Experimental/OnlineServices/Source/OnlineServicesCommon/Private/Online/OnlineSyntaxTests.cpp
--- a/Engine/Plugins/Online/Experimental/OnlineServices/Source/OnlineServicesCommon/Private/Online/OnlineSyntaxTests.cpp
+++ b/Engine/Plugins/Online/Experimental/OnlineServices/Source/OnlineServicesCommon/Private/Online/OnlineSyntaxTests.cpp
@@ -442,21 +442,38 @@ void TestToLogStringSyntax()
 	using UE::Online::ToLogString;
 
 	ToLogString(FString());
+	ToLogString(UE::Online::FAccountInfo());
+	ToLogString(UE::Online::FExternalAuthToken());
+	ToLogString(UE::Online::FExternalServerAuthTicket());
+	ToLogString(UE::Online::FVerifiedAuthSession());
+	ToLogString(UE::Online::FVerifiedAuthTicket());
 	ToLogString(UE::Online::FAuthLogin::Params());
-	//ToLogString(UE::Online::FAuthLogin::Result()); default constructing a TSharedPtr will crash
+	ToLogString(UE::Online::FAuthLogin::Result({MakeShared<UE::Online::FAccountInfo>()}));
 	ToLogString(UE::Online::FAuthLogout::Params());
 	ToLogString(UE::Online::FAuthLogout::Result());
-	ToLogString(UE::Online::FAuthGenerateAuthToken::Params());
-	ToLogString(UE::Online::FAuthGenerateAuthToken::Result());
-	ToLogString(UE::Online::FAuthGetAuthToken::Params());
-	ToLogString(UE::Online::FAuthGetAuthToken::Result());
-	ToLogString(UE::Online::FAuthGenerateAuthCode::Params());
-	ToLogString(UE::Online::FAuthGenerateAuthCode::Result());
-	ToLogString(UE::Online::FAuthGetAccountByPlatformUserId::Params());
-	ToLogString(UE::Online::FAuthGetAccountByPlatformUserId::Result());
-	ToLogString(UE::Online::FAuthGetAccountByAccountId::Params());
-	ToLogString(UE::Online::FAuthGetAccountByAccountId::Result());
-	ToLogString(UE::Online::FLoginStatusChanged());
+	ToLogString(UE::Online::FAuthModifyAccountAttributes::Params());
+	ToLogString(UE::Online::FAuthModifyAccountAttributes::Result());
+	ToLogString(UE::Online::FAuthQueryExternalServerAuthTicket::Params());
+	ToLogString(UE::Online::FAuthQueryExternalServerAuthTicket::Result());
+	ToLogString(UE::Online::FAuthQueryExternalAuthToken::Params());
+	ToLogString(UE::Online::FAuthQueryExternalAuthToken::Result());
+	ToLogString(UE::Online::FAuthQueryVerifiedAuthTicket::Params());
+	ToLogString(UE::Online::FAuthQueryVerifiedAuthTicket::Result());
+	ToLogString(UE::Online::FAuthCancelVerifiedAuthTicket::Params());
+	ToLogString(UE::Online::FAuthCancelVerifiedAuthTicket::Result());
+	ToLogString(UE::Online::FAuthBeginVerifiedAuthSession::Params());
+	ToLogString(UE::Online::FAuthBeginVerifiedAuthSession::Result());
+	ToLogString(UE::Online::FAuthEndVerifiedAuthSession::Params());
+	ToLogString(UE::Online::FAuthEndVerifiedAuthSession::Result());
+	ToLogString(UE::Online::FAuthGetLocalOnlineUserByOnlineAccountId::Params());
+	ToLogString(UE::Online::FAuthGetLocalOnlineUserByOnlineAccountId::Result({MakeShared<UE::Online::FAccountInfo>()}));
+	ToLogString(UE::Online::FAuthGetLocalOnlineUserByPlatformUserId::Params());
+	ToLogString(UE::Online::FAuthGetLocalOnlineUserByPlatformUserId::Result({MakeShared<UE::Online::FAccountInfo>()}));
+	ToLogString(UE::Online::FAuthGetAllLocalOnlineUsers::Params());
+	ToLogString(UE::Online::FAuthGetAllLocalOnlineUsers::Result());
+	ToLogString(UE::Online::FAuthLoginStatusChanged({MakeShared<UE::Online::FAccountInfo>()}));
+	ToLogString(UE::Online::FAuthPendingAuthExpiration({MakeShared<UE::Online::FAccountInfo>()}));
+	ToLogString(UE::Online::FAuthAccountAttributesChanged({MakeShared<UE::Online::FAccountInfo>()}));
 }
 
 /* unnamed namespace */ }
diff --git a/Engine/Plugins/Online/Experimental/OnlineServices/Source/OnlineServicesCommon/Private/Online/SessionsCommon.cpp b/Engine/Plugins/Online/Experimental/OnlineServices/Source/OnlineServicesCommon/Private/Online/SessionsCommon.cpp
--- a/Engine/Plugins/Online/Experimental/OnlineServices/Source/OnlineServicesCommon/Private/Online/SessionsCommon.cpp
+++ b/Engine/Plugins/Online/Experimental/OnlineServices/Source/OnlineServicesCommon/Private/Online/SessionsCommon.cpp
@@ -304,8 +304,8 @@ namespace UE::Online {
 		{
 			IAuthPtr Auth = Services.GetAuthInterface();
 
-			FAuthGetAccountByAccountId::Params GetAccountByAccountIParams{ Params.LocalUserId };
-			TOnlineResult<FAuthGetAccountByAccountId> Result = Auth->GetAccountByAccountId(MoveTemp(GetAccountByAccountIParams));
+			FAuthGetLocalOnlineUserByOnlineAccountId::Params GetLocalOnlineUserByOnlineAccountIdParams{ Params.LocalUserId };
+			TOnlineResult<FAuthGetLocalOnlineUserByOnlineAccountId> Result = Auth->GetLocalOnlineUserByOnlineAccountId(MoveTemp(GetLocalOnlineUserByOnlineAccountIdParams));
 
 			if (Result.IsOk())
 			{
@@ -327,8 +327,8 @@ namespace UE::Online {
 
 		for (const FOnlineAccountIdHandle& LocalUserId : LocalUserIds)
 		{
-			FAuthGetAccountByAccountId::Params GetAccountByAccountIParams{ LocalUserId };
-			TOnlineResult<FAuthGetAccountByAccountId> Result = Auth->GetAccountByAccountId(MoveTemp(GetAccountByAccountIParams));
+			FAuthGetLocalOnlineUserByOnlineAccountId::Params GetLocalOnlineUserByOnlineAccountIdParams{ LocalUserId };
+			TOnlineResult<FAuthGetLocalOnlineUserByOnlineAccountId> Result = Auth->GetLocalOnlineUserByOnlineAccountId(MoveTemp(GetLocalOnlineUserByOnlineAccountIdParams));
 
 			if (Result.IsOk())
 			{
@@ -352,8 +352,8 @@ namespace UE::Online {
 			{
 				IAuthPtr Auth = Services.GetAuthInterface();
 
-				FAuthGetAccountByAccountId::Params GetAccountByAccountIParams{ Params.LocalUserId };
-				TOnlineResult<FAuthGetAccountByAccountId> Result = Auth->GetAccountByAccountId(MoveTemp(GetAccountByAccountIParams));
+				FAuthGetLocalOnlineUserByOnlineAccountId::Params GetLocalOnlineUserByOnlineAccountIdParams{ Params.LocalUserId };
+				TOnlineResult<FAuthGetLocalOnlineUserByOnlineAccountId> Result = Auth->GetLocalOnlineUserByOnlineAccountId(MoveTemp(GetLocalOnlineUserByOnlineAccountIdParams));
 
 				if (Result.IsOk())
 				{
@@ -393,8 +393,8 @@ namespace UE::Online {
 		// User login check
 		IAuthPtr Auth = Services.GetAuthInterface();
 
-		FAuthGetAccountByAccountId::Params GetAccountByAccountIParams{ Params.LocalUserId };
-		TOnlineResult<FAuthGetAccountByAccountId> Result = Auth->GetAccountByAccountId(MoveTemp(GetAccountByAccountIParams));
+		FAuthGetLocalOnlineUserByOnlineAccountId::Params GetLocalOnlineUserByOnlineAccountIdParams{ Params.LocalUserId };
+		TOnlineResult<FAuthGetLocalOnlineUserByOnlineAccountId> Result = Auth->GetLocalOnlineUserByOnlineAccountId(MoveTemp(GetLocalOnlineUserByOnlineAccountIdParams));
 
 		if (Result.IsOk())
 		{
@@ -491,8 +491,8 @@ namespace UE::Online {
 
 		for (const FOnlineAccountIdHandle& LocalUserId : LocalUserIds)
 		{
-			FAuthGetAccountByAccountId::Params GetAccountByAccountIParams{ LocalUserId };
-			TOnlineResult<FAuthGetAccountByAccountId> Result = Auth->GetAccountByAccountId(MoveTemp(GetAccountByAccountIParams));
+			FAuthGetLocalOnlineUserByOnlineAccountId::Params GetLocalOnlineUserByOnlineAccountIdParams{ LocalUserId };
+			TOnlineResult<FAuthGetLocalOnlineUserByOnlineAccountId> Result = Auth->GetLocalOnlineUserByOnlineAccountId(MoveTemp(GetLocalOnlineUserByOnlineAccountIdParams));
 
 			if (Result.IsOk())
 			{
@@ -517,8 +517,8 @@ namespace UE::Online {
 			{
 				IAuthPtr Auth = Services.GetAuthInterface();
 
-				FAuthGetAccountByAccountId::Params GetAccountByAccountIParams{ Params.LocalUserId };
-				TOnlineResult<FAuthGetAccountByAccountId> Result = Auth->GetAccountByAccountId(MoveTemp(GetAccountByAccountIParams));
+				FAuthGetLocalOnlineUserByOnlineAccountId::Params GetLocalOnlineUserByOnlineAccountIdParams{ Params.LocalUserId };
+				TOnlineResult<FAuthGetLocalOnlineUserByOnlineAccountId> Result = Auth->GetLocalOnlineUserByOnlineAccountId(MoveTemp(GetLocalOnlineUserByOnlineAccountIdParams));
 
 				if (Result.IsOk())
 				{
@@ -543,8 +543,8 @@ namespace UE::Online {
 
 		for (const FOnlineAccountIdHandle& LocalUserId : Params.LocalUsers)
 		{
-			FAuthGetAccountByAccountId::Params GetAccountByAccountIParams{ LocalUserId };
-			TOnlineResult<FAuthGetAccountByAccountId> Result = Auth->GetAccountByAccountId(MoveTemp(GetAccountByAccountIParams));
+			FAuthGetLocalOnlineUserByOnlineAccountId::Params GetLocalOnlineUserByOnlineAccountIdParams{ LocalUserId };
+			TOnlineResult<FAuthGetLocalOnlineUserByOnlineAccountId> Result = Auth->GetLocalOnlineUserByOnlineAccountId(MoveTemp(GetLocalOnlineUserByOnlineAccountIdParams));
 
 			if (Result.IsOk())
 			{
diff --git a/Engine/Plugins/Online/Experimental/OnlineServices/Source/OnlineServicesCommon/Public/Online/LobbiesCommon.h b/Engine/Plugins/Online/Experimental/OnlineServices/Source/OnlineServicesCommon/Public/Online/LobbiesCommon.h
--- a/Engine/Plugins/Online/Experimental/OnlineServices/Source/OnlineServicesCommon/Public/Online/LobbiesCommon.h
+++ b/Engine/Plugins/Online/Experimental/OnlineServices/Source/OnlineServicesCommon/Public/Online/LobbiesCommon.h
@@ -9,7 +9,6 @@
 namespace UE::Online {
 
 class FOnlineServicesCommon;
-class FAccountInfo;
 
 class ONLINESERVICESCOMMON_API FLobbiesCommon : public TOnlineComponent<ILobbies>
 {
diff --git a/Engine/Plugins/Online/Experimental/OnlineServices/Source/OnlineServicesCommon/Public/Online/LobbiesCommonTypes.h b/Engine/Plugins/Online/Experimental/OnlineServices/Source/OnlineServicesCommon/Public/Online/LobbiesCommonTypes.h
--- a/Engine/Plugins/Online/Experimental/OnlineServices/Source/OnlineServicesCommon/Public/Online/LobbiesCommonTypes.h
+++ b/Engine/Plugins/Online/Experimental/OnlineServices/Source/OnlineServicesCommon/Public/Online/LobbiesCommonTypes.h
@@ -10,7 +10,6 @@ namespace UE::Online {
 #define LOBBIES_FUNCTIONAL_TEST_ENABLED !UE_BUILD_SHIPPING
 
 class FOnlineServicesCommon;
-class FAccountInfo;
 
 struct FLobbyConfig
 {
diff --git a/Engine/Plugins/Online/Experimental/OnlineServices/Source/OnlineServicesInterface/Public/Online/ExternalUI.h b/Engine/Plugins/Online/Experimental/OnlineServices/Source/OnlineServicesInterface/Public/Online/ExternalUI.h
--- a/Engine/Plugins/Online/Experimental/OnlineServices/Source/OnlineServicesInterface/Public/Online/ExternalUI.h
+++ b/Engine/Plugins/Online/Experimental/OnlineServices/Source/OnlineServicesInterface/Public/Online/ExternalUI.h
@@ -9,7 +9,7 @@
 namespace UE::Online {
 
 class FOnlineError;
-class FAccountInfo;
+struct FAccountInfo;
 
 struct FExternalUIShowLoginUI
 {
diff --git a/Engine/Plugins/Online/Experimental/OnlineServices/Source/OnlineServicesInterface/Public/Online/OnlineUtils.h b/Engine/Plugins/Online/Experimental/OnlineServices/Source/OnlineServicesInterface/Public/Online/OnlineUtils.h
--- a/Engine/Plugins/Online/Experimental/OnlineServices/Source/OnlineServicesInterface/Public/Online/OnlineUtils.h
+++ b/Engine/Plugins/Online/Experimental/OnlineServices/Source/OnlineServicesInterface/Public/Online/OnlineUtils.h
@@ -8,13 +8,16 @@
 namespace UE::Online {
 
 template <typename T> FString ToLogString(const TArray<T>& Array);
+template <typename T> FString ToLogString(const TSet<T>& Set);
 template <typename K, typename V> FString ToLogString(const TMap<K, V>& Map);
+template <typename T, typename U> FString ToLogString(const TPair<T, U>& Pair);
 template <typename T, ESPMode Mode> FString ToLogString(const TSharedPtr<T, Mode>& Ptr);
 template <typename T, ESPMode Mode> FString ToLogString(const TSharedRef<T, Mode>& Ref);
 template <typename T> FString ToLogString(const TOptional<T> Optional);
 template <typename... Ts> FString ToLogString(const TVariant<Ts...>& Variant);
 template <typename T, ESPMode Mode> FString ToLogString(const TSharedRef<T, Mode>& Ref);
 inline FString ToLogString(const FString& String);
+inline FString ToLogString(const FName& Name);
 inline FString ToLogString(uint8 Value);
 inline FString ToLogString(int8 Value);
 inline FString ToLogString(uint16 Value);
@@ -23,9 +26,9 @@ inline FString ToLogString(uint32 Value);
 inline FString ToLogString(int32 Value);
 inline FString ToLogString(uint64 Value);
 inline FString ToLogString(int64 Value);
+inline FString ToLogString(bool Value);
 inline FString ToLogString(float Value);
 inline FString ToLogString(double Value);
-inline FString ToLogString(bool Value);
 template <typename T> std::enable_if_t<!TModels<Meta::COnlineMetadataAvailable, T>::Value, FString> ToLogString(const T& Value);
 template <typename T> std::enable_if_t<TModels<Meta::COnlineMetadataAvailable, T>::Value, FString> ToLogString(const T& Value);
 
@@ -51,6 +54,28 @@ FString ToLogString(const TArray<T>& Array)
 	return LogString;
 }
 
+template <typename T>
+FString ToLogString(const TSet<T>& Set)
+{
+	FString LogString;
+	LogString += TEXT("[ ");
+	bool bFirst = true;
+	for (const T& Value : Set)
+	{
+		if (bFirst)
+		{
+			bFirst = false;
+		}
+		else
+		{
+			LogString += TEXT(", ");
+		}
+		LogString += ToLogString(Value);
+	}
+	LogString += TEXT(" ]");
+	return LogString;
+}
+
 template <typename K, typename V>
 FString ToLogString(const TMap<K, V>& Map)
 {
@@ -75,6 +100,18 @@ FString ToLogString(const TMap<K, V>& Map)
 	return LogString;
 }
 
+template <typename T, typename U>
+FString ToLogString(const TPair<T, U>& Pair)
+{
+	FString LogString;
+	LogString += TEXT("[ ");
+	LogString += ToLogString(Pair.template Get<0>());
+	LogString += TEXT(", ");
+	LogString += ToLogString(Pair.template Get<1>());
+	LogString += TEXT(" ]");
+	return LogString;
+}
+
 template <typename T, ESPMode Mode>
 FString ToLogString(const TSharedPtr<T, Mode>& Ptr)
 {
@@ -121,6 +158,11 @@ inline FString ToLogString(const FString& String)
 	return String;
 }
 
+inline FString ToLogString(const FName& Name)
+{
+	return Name.ToString();
+}
+
 inline FString ToLogString(uint8 Value)
 {
 	return FString::Printf(TEXT("%hhu"), Value);
@@ -212,4 +254,12 @@ std::enable_if_t<TModels<Meta::COnlineMetadataAvailable, T>::Value, FString> ToL
 	return LogString;
 }
 
+template <typename DataType, typename OpType>
+const DataType& GetOpDataChecked(const TOnlineAsyncOp<OpType>& Op, const FString& Key)
+{
+	const DataType* Data = Op.Data.template Get<DataType>(Key);
+	check(Data);
+	return *Data;
+}
+
 /* UE::Online */ }
\ No newline at end of file
diff --git a/Engine/Plugins/Online/Experimental/OnlineServices/Source/OnlineServicesInterface/Public/Online/Stats.h b/Engine/Plugins/Online/Experimental/OnlineServices/Source/OnlineServicesInterface/Public/Online/Stats.h
--- a/Engine/Plugins/Online/Experimental/OnlineServices/Source/OnlineServicesInterface/Public/Online/Stats.h
+++ b/Engine/Plugins/Online/Experimental/OnlineServices/Source/OnlineServicesInterface/Public/Online/Stats.h
@@ -5,13 +5,11 @@
 #include "Online/CoreOnline.h"
 #include "Online/OnlineAsyncOpHandle.h"
 #include "Online/OnlineMeta.h"
-
-// TODO: Include Online/SchemaVariant.h instead when it's ready
-#include "Online/Lobbies.h"
+#include "Online/Schema.h"
 
 namespace UE::Online {
 
-using FStatValue = FLobbyVariant;
+using FStatValue = FSchemaVariant;
 
 struct FUserStats
 {
diff --git a/Engine/Plugins/Online/Experimental/OnlineServicesEOSGS/Source/Private/Online/LobbiesEOSGS.cpp b/Engine/Plugins/Online/Experimental/OnlineServicesEOSGS/Source/Private/Online/LobbiesEOSGS.cpp
--- a/Engine/Plugins/Online/Experimental/OnlineServicesEOSGS/Source/Private/Online/LobbiesEOSGS.cpp
+++ b/Engine/Plugins/Online/Experimental/OnlineServicesEOSGS/Source/Private/Online/LobbiesEOSGS.cpp
@@ -22,14 +22,6 @@ static const FString LobbyMemberChangesKeyName = TEXT("LobbyMemberChanges");
 
 static const int32 MaxAttributeSize = 1000;
 
-template <typename DataType, typename OpType>
-const DataType& GetOpDataChecked(const TOnlineAsyncOp<OpType>& Op, const FString& Key)
-{
-	const DataType* Data = Op.Data.template Get<DataType>(Key);
-	check(Data);
-	return *Data;
-}
-
 FLobbiesEOSGS::FLobbiesEOSGS(FOnlineServicesEOSGS& InServices)
 	: FLobbiesCommon(InServices)
 {
diff --git a/Engine/Plugins/Online/Experimental/OnlineServicesEOSGS/Source/Private/Online/SessionsEOSGS.cpp b/Engine/Plugins/Online/Experimental/OnlineServicesEOSGS/Source/Private/Online/SessionsEOSGS.cpp
--- a/Engine/Plugins/Online/Experimental/OnlineServicesEOSGS/Source/Private/Online/SessionsEOSGS.cpp
+++ b/Engine/Plugins/Online/Experimental/OnlineServicesEOSGS/Source/Private/Online/SessionsEOSGS.cpp
@@ -127,8 +127,8 @@ TOnlineAsyncOpHandle<FCreateSession> FSessionsEOSGS::CreateSession(FCreateSessio
 		{
 			IAuthPtr Auth = Services.GetAuthInterface();
 
-			FAuthGetAccountByAccountId::Params GetAccountByAccountIParams { OpParams.LocalUserId };
-			TOnlineResult<FAuthGetAccountByAccountId> Result = Auth->GetAccountByAccountId(MoveTemp(GetAccountByAccountIParams));
+			FAuthGetLocalOnlineUserByOnlineAccountId::Params GetLocalOnlineUserByOnlineAccountIdParams { OpParams.LocalUserId };
+			TOnlineResult<FAuthGetLocalOnlineUserByOnlineAccountId> Result = Auth->GetLocalOnlineUserByOnlineAccountId(MoveTemp(GetLocalOnlineUserByOnlineAccountIdParams));
 
 			if (Result.IsOk())
 			{
@@ -550,7 +550,7 @@ TOnlineAsyncOpHandle<FUpdateSession> FSessionsEOSGS::UpdateSession(FUpdateSessio
 			{
 				IAuthPtr Auth = Services.GetAuthInterface();
 
-				TOnlineResult<FAuthGetAccountByAccountId> Result = Auth->GetAccountByAccountId({ OpParams.LocalUserId });
+				TOnlineResult<FAuthGetLocalOnlineUserByOnlineAccountId> Result = Auth->GetLocalOnlineUserByOnlineAccountId({ OpParams.LocalUserId });
 
 				if (Result.IsOk())
 				{
@@ -788,7 +788,7 @@ TOnlineAsyncOpHandle<FFindSessions> FSessionsEOSGS::FindSessions(FFindSessions::
 
 	// We check if the user is logged in
 	IAuthPtr Auth = Services.GetAuthInterface();
-	TOnlineResult<FAuthGetAccountByAccountId> Result = Auth->GetAccountByAccountId({ OpParams.LocalUserId });
+	TOnlineResult<FAuthGetLocalOnlineUserByOnlineAccountId> Result = Auth->GetLocalOnlineUserByOnlineAccountId({ OpParams.LocalUserId });
 
 	if (Result.IsOk())
 	{
diff --git a/Engine/Plugins/Online/Experimental/OnlineServicesEOSGS/Source/Private/Online/SocketSubsystemEOSUtils_OnlineServicesEOSGS.cpp b/Engine/Plugins/Online/Experimental/OnlineServicesEOSGS/Source/Private/Online/SocketSubsystemEOSUtils_OnlineServicesEOSGS.cpp
--- a/Engine/Plugins/Online/Experimental/OnlineServicesEOSGS/Source/Private/Online/SocketSubsystemEOSUtils_OnlineServicesEOSGS.cpp
+++ b/Engine/Plugins/Online/Experimental/OnlineServicesEOSGS/Source/Private/Online/SocketSubsystemEOSUtils_OnlineServicesEOSGS.cpp
@@ -27,14 +27,14 @@ EOS_ProductUserId FSocketSubsystemEOSUtils_OnlineServicesEOS::GetLocalUserId()
 	IAuthPtr AuthEOS = ServicesEOSGS.GetAuthInterface();
 	check(AuthEOS);
 
-	FAuthGetAccountByPlatformUserId::Params AuthParams;
+	FAuthGetLocalOnlineUserByPlatformUserId::Params AuthParams;
 	AuthParams.PlatformUserId = FPlatformMisc::GetPlatformUserForUserIndex(0);
-	TOnlineResult<FAuthGetAccountByPlatformUserId> AuthResult = AuthEOS->GetAccountByPlatformUserId(MoveTemp(AuthParams));
+	TOnlineResult<FAuthGetLocalOnlineUserByPlatformUserId> AuthResult = AuthEOS->GetLocalOnlineUserByPlatformUserId(MoveTemp(AuthParams));
 	if (AuthResult.IsOk())
 	{
-		UE::Online::FAuthGetAccountByPlatformUserId::Result OkValue = AuthResult.GetOkValue();
+		UE::Online::FAuthGetLocalOnlineUserByPlatformUserId::Result OkValue = AuthResult.GetOkValue();
 
-		Result = GetProductUserIdChecked(OkValue.AccountInfo->UserId);
+		Result = GetProductUserIdChecked(OkValue.AccountInfo->AccountId);
 	}
 	else
 	{
@@ -51,18 +51,18 @@ FString FSocketSubsystemEOSUtils_OnlineServicesEOS::GetSessionId()
 	IAuthPtr AuthEOS = ServicesEOSGS.GetAuthInterface();
 	check(AuthEOS);
 
-	FAuthGetAccountByPlatformUserId::Params AuthParams;
+	FAuthGetLocalOnlineUserByPlatformUserId::Params AuthParams;
 	AuthParams.PlatformUserId = FPlatformMisc::GetPlatformUserForUserIndex(0);
-	TOnlineResult<FAuthGetAccountByPlatformUserId> AuthResult = AuthEOS->GetAccountByPlatformUserId(MoveTemp(AuthParams));
+	TOnlineResult<FAuthGetLocalOnlineUserByPlatformUserId> AuthResult = AuthEOS->GetLocalOnlineUserByPlatformUserId(MoveTemp(AuthParams));
 	if (AuthResult.IsOk())
 	{
-		FAuthGetAccountByPlatformUserId::Result* AuthOkValue = AuthResult.TryGetOkValue();
+		FAuthGetLocalOnlineUserByPlatformUserId::Result* AuthOkValue = AuthResult.TryGetOkValue();
 
 		ILobbiesPtr LobbiesEOS = ServicesEOSGS.GetLobbiesInterface();
 		check(LobbiesEOS);
 
 		FGetJoinedLobbies::Params LobbiesParams;
-		LobbiesParams.LocalUserId = AuthOkValue->AccountInfo->UserId;
+		LobbiesParams.LocalUserId = AuthOkValue->AccountInfo->AccountId;
 		TOnlineResult<FGetJoinedLobbies> LobbiesResult = LobbiesEOS->GetJoinedLobbies(MoveTemp(LobbiesParams));
 		if (LobbiesResult.IsOk())
 		{
diff --git a/Engine/Plugins/Online/Experimental/OnlineServicesOSSAdapter/Source/Private/Online/AchievementsOSSAdapter.cpp b/Engine/Plugins/Online/Experimental/OnlineServicesOSSAdapter/Source/Private/Online/AchievementsOSSAdapter.cpp
--- a/Engine/Plugins/Online/Experimental/OnlineServicesOSSAdapter/Source/Private/Online/AchievementsOSSAdapter.cpp
+++ b/Engine/Plugins/Online/Experimental/OnlineServicesOSSAdapter/Source/Private/Online/AchievementsOSSAdapter.cpp
@@ -56,8 +56,8 @@ TOnlineAsyncOpHandle<FQueryAchievementDefinitions> FAchievementsOSSAdapter::Quer
 		TPromise<void> Promise;
 		TFuture<void> Future = Promise.GetFuture();
 
-		const FUniqueNetIdRef LocalUserId = Services.Get<FAuthOSSAdapter>()->GetUniqueNetId(Op.GetParams().LocalUserId);
-		if (!LocalUserId->IsValid())
+		const FUniqueNetIdPtr LocalUserId = Services.Get<FAuthOSSAdapter>()->GetUniqueNetId(Op.GetParams().LocalUserId);
+		if (!LocalUserId)
 		{
 			Op.SetError(Errors::InvalidUser());
 			Promise.SetValue();
@@ -90,7 +90,12 @@ TOnlineAsyncOpHandle<FQueryAchievementDefinitions> FAchievementsOSSAdapter::Quer
 		TPromise<void> Promise;
 		TFuture<void> Future = Promise.GetFuture();
 
-		const FUniqueNetIdRef LocalUserId = Services.Get<FAuthOSSAdapter>()->GetUniqueNetId(Op.GetParams().LocalUserId);
+		const FUniqueNetIdPtr LocalUserId = Services.Get<FAuthOSSAdapter>()->GetUniqueNetId(Op.GetParams().LocalUserId);
+		if (!LocalUserId)
+		{
+			Op.SetError(Errors::InvalidUser());
+			return Future;
+		}
 
 		AchievementsInterface->QueryAchievementDescriptions(*LocalUserId, *MakeDelegateAdapter(this, [this, Promise = MoveTemp(Promise), WeakOp = Op.AsWeak()](const FUniqueNetId& LocalUserId, const bool bWasSuccessful) mutable
 		{
@@ -108,7 +113,12 @@ TOnlineAsyncOpHandle<FQueryAchievementDefinitions> FAchievementsOSSAdapter::Quer
 	})
 	.Then([this](TOnlineAsyncOp<FQueryAchievementDefinitions>& Op)
 	{
-		const FUniqueNetIdRef LocalUserId = Services.Get<FAuthOSSAdapter>()->GetUniqueNetId(Op.GetParams().LocalUserId);
+		const FUniqueNetIdPtr LocalUserId = Services.Get<FAuthOSSAdapter>()->GetUniqueNetId(Op.GetParams().LocalUserId);
+		if (!LocalUserId)
+		{
+			Op.SetError(Errors::InvalidUser());
+			return;
+		}
 
 		TArray<FOnlineAchievement> Achievements;
 		EOnlineCachedResult::Type Result = AchievementsInterface->GetCachedAchievements(*LocalUserId, Achievements);
@@ -150,8 +160,8 @@ TOnlineAsyncOpHandle<FQueryAchievementDefinitions> FAchievementsOSSAdapter::Quer
 
 TOnlineResult<FGetAchievementIds> FAchievementsOSSAdapter::GetAchievementIds(FGetAchievementIds::Params&& Params)
 {
-	const FUniqueNetIdRef LocalUserId = Services.Get<FAuthOSSAdapter>()->GetUniqueNetId(Params.LocalUserId);
-	if (!LocalUserId->IsValid())
+	const FUniqueNetIdPtr LocalUserId = Services.Get<FAuthOSSAdapter>()->GetUniqueNetId(Params.LocalUserId);
+	if (!LocalUserId)
 	{
 		return TOnlineResult<FGetAchievementIds>(Errors::InvalidUser());
 	}
@@ -169,8 +179,8 @@ TOnlineResult<FGetAchievementIds> FAchievementsOSSAdapter::GetAchievementIds(FGe
 
 TOnlineResult<FGetAchievementDefinition> FAchievementsOSSAdapter::GetAchievementDefinition(FGetAchievementDefinition::Params&& Params)
 {
-	const FUniqueNetIdRef LocalUserId = Services.Get<FAuthOSSAdapter>()->GetUniqueNetId(Params.LocalUserId);
-	if (!LocalUserId->IsValid())
+	const FUniqueNetIdPtr LocalUserId = Services.Get<FAuthOSSAdapter>()->GetUniqueNetId(Params.LocalUserId);
+	if (!LocalUserId)
 	{
 		return TOnlineResult<FGetAchievementDefinition>(Errors::InvalidUser());
 	}
@@ -199,8 +209,8 @@ TOnlineAsyncOpHandle<FQueryAchievementStates> FAchievementsOSSAdapter::QueryAchi
 		TPromise<void> Promise;
 		TFuture<void> Future = Promise.GetFuture();
 
-		const FUniqueNetIdRef LocalUserId = Services.Get<FAuthOSSAdapter>()->GetUniqueNetId(Op.GetParams().LocalUserId);
-		if (!LocalUserId->IsValid())
+		const FUniqueNetIdPtr LocalUserId = Services.Get<FAuthOSSAdapter>()->GetUniqueNetId(Op.GetParams().LocalUserId);
+		if (!LocalUserId)
 		{
 			Op.SetError(Errors::InvalidUser());
 			Promise.SetValue();
@@ -231,7 +241,13 @@ TOnlineAsyncOpHandle<FQueryAchievementStates> FAchievementsOSSAdapter::QueryAchi
 	})
 	.Then([this](TOnlineAsyncOp<FQueryAchievementStates>& Op)
 	{
-		const FUniqueNetIdRef LocalUserId = Services.Get<FAuthOSSAdapter>()->GetUniqueNetId(Op.GetParams().LocalUserId);
+		const FUniqueNetIdPtr LocalUserId = Services.Get<FAuthOSSAdapter>()->GetUniqueNetId(Op.GetParams().LocalUserId);
+		if (!LocalUserId)
+		{
+			Op.SetError(Errors::InvalidUser());
+			return;
+		}
+
 		TArray<FOnlineAchievement> Achievements;
 		EOnlineCachedResult::Type Result = AchievementsInterface->GetCachedAchievements(*LocalUserId, Achievements);
 		if (Result != EOnlineCachedResult::Success)
@@ -259,8 +275,8 @@ TOnlineAsyncOpHandle<FQueryAchievementStates> FAchievementsOSSAdapter::QueryAchi
 
 TOnlineResult<FGetAchievementState> FAchievementsOSSAdapter::GetAchievementState(FGetAchievementState::Params&& Params)
 {
-	const FUniqueNetIdRef LocalUserId = Services.Get<FAuthOSSAdapter>()->GetUniqueNetId(Params.LocalUserId);
-	if (!LocalUserId->IsValid())
+	const FUniqueNetIdPtr LocalUserId = Services.Get<FAuthOSSAdapter>()->GetUniqueNetId(Params.LocalUserId);
+	if (!LocalUserId)
 	{
 		return TOnlineResult<FGetAchievementState>(Errors::InvalidUser());
 	}
diff --git a/Engine/Plugins/Online/Experimental/OnlineServicesOSSAdapter/Source/Private/Online/ExternalUIOSSAdapter.cpp b/Engine/Plugins/Online/Experimental/OnlineServicesOSSAdapter/Source/Private/Online/ExternalUIOSSAdapter.cpp
--- a/Engine/Plugins/Online/Experimental/OnlineServicesOSSAdapter/Source/Private/Online/ExternalUIOSSAdapter.cpp
+++ b/Engine/Plugins/Online/Experimental/OnlineServicesOSSAdapter/Source/Private/Online/ExternalUIOSSAdapter.cpp
@@ -69,7 +69,7 @@ TOnlineAsyncOpHandle<FExternalUIShowLoginUI> FExternalUIOSSAdapter::ShowLoginUI(
 									FOnlineAccountIdHandle Handle = static_cast<FOnlineServicesOSSAdapter&>(Services).GetAccountIdRegistry().FindOrAddHandle(UniqueId.ToSharedRef());
 									FExternalUIShowLoginUI::Result Result = { MakeShared<FAccountInfo>() };
 									Result.AccountInfo->PlatformUserId = PinnedOp->GetParams().PlatformUserId;
-									Result.AccountInfo->UserId = Handle;
+									Result.AccountInfo->AccountId = Handle;
 									Result.AccountInfo->LoginStatus = ELoginStatus::LoggedIn;
 
 									PinnedOp->SetResult(MoveTemp(Result));
diff --git a/Engine/Plugins/Online/Experimental/OnlineServicesOSSAdapter/Source/Private/Online/PresenceOSSAdapter.cpp b/Engine/Plugins/Online/Experimental/OnlineServicesOSSAdapter/Source/Private/Online/PresenceOSSAdapter.cpp
--- a/Engine/Plugins/Online/Experimental/OnlineServicesOSSAdapter/Source/Private/Online/PresenceOSSAdapter.cpp
+++ b/Engine/Plugins/Online/Experimental/OnlineServicesOSSAdapter/Source/Private/Online/PresenceOSSAdapter.cpp
@@ -146,9 +146,9 @@ TOnlineAsyncOpHandle<FQueryPresence> FPresenceOSSAdapter::QueryPresence(FQueryPr
 
 	Op->Then([this](TOnlineAsyncOp<FQueryPresence>& Op, TPromise<TOnlineResult<FQueryPresence>>&& Result)
 	{
-		const FUniqueNetIdRef UniqueNetId = Auth->GetUniqueNetId(Op.GetParams().LocalUserId);
+		const FUniqueNetIdPtr UniqueNetId = Auth->GetUniqueNetId(Op.GetParams().LocalUserId);
 
-		if (UniqueNetId->IsValid())
+		if (UniqueNetId)
 		{
 			PresenceInt->QueryPresence(*UniqueNetId, *MakeDelegateAdapter(this, [this, ResultPromise = MoveTemp(Result)](const FUniqueNetId& UserId, const bool bWasSuccessful) mutable
 			{
@@ -179,11 +179,24 @@ TOnlineAsyncOpHandle<FBatchQueryPresence> FPresenceOSSAdapter::BatchQueryPresenc
 
 	Op->Then([this](TOnlineAsyncOp<FBatchQueryPresence>& Op, TPromise<TOnlineResult<FBatchQueryPresence>>&& Result)
 	{
-		const FUniqueNetIdRef LocalUserId = Auth->GetUniqueNetId(Op.GetParams().LocalUserId);
+		const FUniqueNetIdPtr LocalUserId = Auth->GetUniqueNetId(Op.GetParams().LocalUserId);
+		if (!LocalUserId)
+		{
+			Op.SetError(Errors::InvalidUser());
+			return;
+		}
+
 		TArray<FUniqueNetIdRef> NetIds;
 		for (const FOnlineAccountIdHandle& TargetUserId : Op.GetParams().TargetUserIds)
 		{
-			NetIds.Add(Auth->GetUniqueNetId(TargetUserId));
+			const FUniqueNetIdPtr TargetUserNetId = Auth->GetUniqueNetId(TargetUserId);
+			if (!TargetUserNetId)
+			{
+				Op.SetError(Errors::InvalidParams());
+				return;
+			}
+
+			NetIds.Add(TargetUserNetId.ToSharedRef());
 		}
 
 		PresenceInt->QueryPresence(*LocalUserId, NetIds, *MakeDelegateAdapter(this, [this, ResultPromise=MoveTemp(Result)](const FUniqueNetId& UserId, const bool bWasSuccessful) mutable
diff --git a/Engine/Plugins/Online/Experimental/OnlineServicesOSSAdapter/Source/Private/Online/StatsOSSAdapter.cpp b/Engine/Plugins/Online/Experimental/OnlineServicesOSSAdapter/Source/Private/Online/StatsOSSAdapter.cpp
--- a/Engine/Plugins/Online/Experimental/OnlineServicesOSSAdapter/Source/Private/Online/StatsOSSAdapter.cpp
+++ b/Engine/Plugins/Online/Experimental/OnlineServicesOSSAdapter/Source/Private/Online/StatsOSSAdapter.cpp
@@ -85,13 +85,13 @@ FOnlineError FStatsOSSAdapter::ConvertUpdateUsersStatsV2toV1(const TArray<FUserS
 {
 	for (const FUserStats& UserStats : UpdateUsersStats)
 	{
-		const FUniqueNetIdRef UserId = Auth->GetUniqueNetId(UserStats.UserId);
-		if (!UserId->IsValid())
+		const FUniqueNetIdPtr UserId = Auth->GetUniqueNetId(UserStats.UserId);
+		if (!UserId)
 		{
 			return Errors::InvalidUser();
 		}
 
-		FOnlineStatsUserUpdatedStats OnlineStatsUserUpdatedStats(UserId);
+		FOnlineStatsUserUpdatedStats OnlineStatsUserUpdatedStats(UserId.ToSharedRef());
 		for (const TPair<FString, FStatValue>& StatPair : UserStats.Stats)
 		{
 			FOnlineStatUpdate OnlineStatUpdate;
@@ -111,8 +111,8 @@ TOnlineAsyncOpHandle<FUpdateStats> FStatsOSSAdapter::UpdateStats(FUpdateStats::P
 
 	Op->Then([this](TOnlineAsyncOp<FUpdateStats>& Op)
 	{
-		const FUniqueNetIdRef LocalUserId = Auth->GetUniqueNetId(Op.GetParams().LocalUserId);
-		if (!LocalUserId->IsValid())
+		const FUniqueNetIdPtr LocalUserId = Auth->GetUniqueNetId(Op.GetParams().LocalUserId);
+		if (!LocalUserId)
 		{
 			Op.SetError(Errors::InvalidUser());
 			return;
@@ -126,7 +126,7 @@ TOnlineAsyncOpHandle<FUpdateStats> FStatsOSSAdapter::UpdateStats(FUpdateStats::P
 			return;
 		}
 
-		StatsInterface->UpdateStats(LocalUserId, UpdateUserStatsV1, *MakeDelegateAdapter(this, [WeakOp = Op.AsWeak()](const FOnlineErrorOss& OnlineError) mutable
+		StatsInterface->UpdateStats(LocalUserId.ToSharedRef(), UpdateUserStatsV1, *MakeDelegateAdapter(this, [WeakOp = Op.AsWeak()](const FOnlineErrorOss& OnlineError) mutable
 		{
 			if (TSharedPtr<TOnlineAsyncOp<FUpdateStats>> PinnedOp = WeakOp.Pin())
 			{
@@ -151,16 +151,16 @@ TOnlineAsyncOpHandle<FQueryStats> FStatsOSSAdapter::QueryStats(FQueryStats::Para
 
 	Op->Then([this](TOnlineAsyncOp<FQueryStats>& Op)
 	{
-		const FUniqueNetIdRef LocalUniqueNetId = Auth->GetUniqueNetId(Op.GetParams().LocalUserId);
-		const FUniqueNetIdRef TargetUniqueNetId = Auth->GetUniqueNetId(Op.GetParams().TargetUserId);
+		const FUniqueNetIdPtr LocalUniqueNetId = Auth->GetUniqueNetId(Op.GetParams().LocalUserId);
+		const FUniqueNetIdPtr TargetUniqueNetId = Auth->GetUniqueNetId(Op.GetParams().TargetUserId);
 
-		if (!LocalUniqueNetId->IsValid() || TargetUniqueNetId->IsValid())
+		if (!LocalUniqueNetId || !TargetUniqueNetId)
 		{
 			Op.SetError(Errors::InvalidUser());
 			return;
 		}
 
-		StatsInterface->QueryStats(LocalUniqueNetId, TargetUniqueNetId, *MakeDelegateAdapter(this, [this, WeakOp = Op.AsWeak()](const FOnlineErrorOss& OnlineError, const TSharedPtr<const FOnlineStatsUserStats>& QueriedStats) mutable
+		StatsInterface->QueryStats(LocalUniqueNetId.ToSharedRef(), TargetUniqueNetId.ToSharedRef(), *MakeDelegateAdapter(this, [this, WeakOp = Op.AsWeak()](const FOnlineErrorOss& OnlineError, const TSharedPtr<const FOnlineStatsUserStats>& QueriedStats) mutable
 		{
 			if (TSharedPtr<TOnlineAsyncOp<FQueryStats>> PinnedOp = WeakOp.Pin())
 			{
@@ -175,7 +175,7 @@ TOnlineAsyncOpHandle<FQueryStats> FStatsOSSAdapter::QueryStats(FQueryStats::Para
 				{
 					FStatValue StatValue;
 					ConvertStatValueV1ToStatValueV2(StatPair.Value, StatValue);
-					Result.Stats.Emplace(StatPair.Key, StatValue);
+					Result.Stats.Emplace(StatPair.Key, MoveTemp(StatValue));
 				}
 
 				PinnedOp->SetResult(MoveTemp(Result));
@@ -193,21 +193,27 @@ TOnlineAsyncOpHandle<FBatchQueryStats> FStatsOSSAdapter::BatchQueryStats(FBatchQ
 
 	Op->Then([this](TOnlineAsyncOp<FBatchQueryStats>& Op)
 	{
-		const FUniqueNetIdRef LocalUserId = Auth->GetUniqueNetId(Op.GetParams().LocalUserId);
+		const FUniqueNetIdPtr LocalUserId = Auth->GetUniqueNetId(Op.GetParams().LocalUserId);
+		if (!LocalUserId)
+		{
+			Op.SetError(Errors::InvalidUser());
+			return;
+		}
+
 		TArray<FUniqueNetIdRef> NetIds;
 		for (const FOnlineAccountIdHandle& TargetUserId : Op.GetParams().TargetUserIds)
 		{
-			const FUniqueNetIdRef UniqueNetId = Auth->GetUniqueNetId(Op.GetParams().LocalUserId);
-			if (!UniqueNetId->IsValid())
+			const FUniqueNetIdPtr UniqueNetId = Auth->GetUniqueNetId(Op.GetParams().LocalUserId);
+			if (!UniqueNetId)
 			{
 				Op.SetError(Errors::InvalidUser());
 				return;
 			}
 
-			NetIds.Add(UniqueNetId);
+			NetIds.Add(UniqueNetId.ToSharedRef());
 		}
 
-		StatsInterface->QueryStats(LocalUserId, NetIds, Op.GetParams().StatNames, *MakeDelegateAdapter(this, [this, WeakOp = Op.AsWeak()](const FOnlineErrorOss& OnlineError, const TArray<TSharedRef<const FOnlineStatsUserStats>>& UsersStatsResult) mutable
+		StatsInterface->QueryStats(LocalUserId.ToSharedRef(), NetIds, Op.GetParams().StatNames, *MakeDelegateAdapter(this, [this, WeakOp = Op.AsWeak()](const FOnlineErrorOss& OnlineError, const TArray<TSharedRef<const FOnlineStatsUserStats>>& UsersStatsResult) mutable
 		{
 			if (TSharedPtr<TOnlineAsyncOp<FBatchQueryStats>> PinnedOp = WeakOp.Pin())
 			{
@@ -227,7 +233,7 @@ TOnlineAsyncOpHandle<FBatchQueryStats> FStatsOSSAdapter::BatchQueryStats(FBatchQ
 					{
 						FStatValue StatValue;
 						ConvertStatValueV1ToStatValueV2(StatPair.Value, StatValue);
-						UserStats.Stats.Emplace(StatPair.Key, StatValue);
+						UserStats.Stats.Emplace(StatPair.Key, MoveTemp(StatValue));
 					}
 				}
 
@@ -247,8 +253,14 @@ TOnlineAsyncOpHandle<FResetStats> FStatsOSSAdapter::ResetStats(FResetStats::Para
 
 	Op->Then([this](TOnlineAsyncOp<FResetStats>& Op)
 	{
-		const FUniqueNetIdRef LocalUserId = Auth->GetUniqueNetId(Op.GetParams().LocalUserId);
-		StatsInterface->ResetStats(LocalUserId);
+		const FUniqueNetIdPtr LocalUserId = Auth->GetUniqueNetId(Op.GetParams().LocalUserId);
+		if (!LocalUserId)
+		{
+			Op.SetError(Errors::InvalidUser());
+			return;
+		}
+
+		StatsInterface->ResetStats(LocalUserId.ToSharedRef());
 		Op.SetResult({});
 	})
 	.Enqueue(GetSerialQueue());
diff --git a/Engine/Plugins/Online/Experimental/OnlineServicesOSSAdapter/Source/Private/Online/UserFileOSSAdapter.cpp b/Engine/Plugins/Online/Experimental/OnlineServicesOSSAdapter/Source/Private/Online/UserFileOSSAdapter.cpp
--- a/Engine/Plugins/Online/Experimental/OnlineServicesOSSAdapter/Source/Private/Online/UserFileOSSAdapter.cpp
+++ b/Engine/Plugins/Online/Experimental/OnlineServicesOSSAdapter/Source/Private/Online/UserFileOSSAdapter.cpp
@@ -35,8 +35,8 @@ TOnlineAsyncOpHandle<FUserFileEnumerateFiles> FUserFileOSSAdapter::EnumerateFile
 
 	Op->Then([this](TOnlineAsyncOp<FUserFileEnumerateFiles>& Op)
 	{
-		const FUniqueNetIdRef LocalUserId = Services.Get<FAuthOSSAdapter>()->GetUniqueNetId(Op.GetParams().LocalUserId);
-		if (!LocalUserId->IsValid())
+		const FUniqueNetIdPtr LocalUserId = Services.Get<FAuthOSSAdapter>()->GetUniqueNetId(Op.GetParams().LocalUserId);
+		if (!LocalUserId)
 		{
 			Op.SetError(Errors::InvalidUser());
 			return MakeFulfilledPromise<void>().GetFuture();
@@ -69,8 +69,8 @@ TOnlineAsyncOpHandle<FUserFileEnumerateFiles> FUserFileOSSAdapter::EnumerateFile
 	})
 	.Then([this](TOnlineAsyncOp<FUserFileEnumerateFiles>& Op)
 	{
-		const FUniqueNetIdRef LocalUserId = Services.Get<FAuthOSSAdapter>()->GetUniqueNetId(Op.GetParams().LocalUserId);
-		if (!LocalUserId->IsValid())
+		const FUniqueNetIdPtr LocalUserId = Services.Get<FAuthOSSAdapter>()->GetUniqueNetId(Op.GetParams().LocalUserId);
+		if (!LocalUserId)
 		{
 			Op.SetError(Errors::InvalidUser());
 			return;
@@ -96,8 +96,8 @@ TOnlineAsyncOpHandle<FUserFileEnumerateFiles> FUserFileOSSAdapter::EnumerateFile
 
 TOnlineResult<FUserFileGetEnumeratedFiles> FUserFileOSSAdapter::GetEnumeratedFiles(FUserFileGetEnumeratedFiles::Params&& Params)
 {
-	const FUniqueNetIdRef LocalUserId = Services.Get<FAuthOSSAdapter>()->GetUniqueNetId(Params.LocalUserId);
-	if (!LocalUserId->IsValid())
+	const FUniqueNetIdPtr LocalUserId = Services.Get<FAuthOSSAdapter>()->GetUniqueNetId(Params.LocalUserId);
+	if (!LocalUserId)
 	{
 		return TOnlineResult<FUserFileGetEnumeratedFiles>(Errors::InvalidUser());
 	}
@@ -129,8 +129,8 @@ TOnlineAsyncOpHandle<FUserFileReadFile> FUserFileOSSAdapter::ReadFile(FUserFileR
 	{
 		const FUserFileReadFile::Params& Params = Op.GetParams();
 
-		const FUniqueNetIdRef LocalUserId = Services.Get<FAuthOSSAdapter>()->GetUniqueNetId(Params.LocalUserId);
-		if (!LocalUserId->IsValid())
+		const FUniqueNetIdPtr LocalUserId = Services.Get<FAuthOSSAdapter>()->GetUniqueNetId(Params.LocalUserId);
+		if (!LocalUserId)
 		{
 			Op.SetError(Errors::InvalidUser());
 			return MakeFulfilledPromise<void>().GetFuture();
@@ -180,8 +180,8 @@ TOnlineAsyncOpHandle<FUserFileReadFile> FUserFileOSSAdapter::ReadFile(FUserFileR
 	{
 		const FUserFileReadFile::Params& Params = Op.GetParams();
 
-		const FUniqueNetIdRef LocalUserId = Services.Get<FAuthOSSAdapter>()->GetUniqueNetId(Params.LocalUserId);
-		if (!LocalUserId->IsValid())
+		const FUniqueNetIdPtr LocalUserId = Services.Get<FAuthOSSAdapter>()->GetUniqueNetId(Params.LocalUserId);
+		if (!LocalUserId)
 		{
 			Op.SetError(Errors::InvalidUser());
 			return;
@@ -217,8 +217,8 @@ TOnlineAsyncOpHandle<FUserFileWriteFile> FUserFileOSSAdapter::WriteFile(FUserFil
 	{
 		const FUserFileWriteFile::Params& Params = Op.GetParams();
 
-		const FUniqueNetIdRef LocalUserId = Services.Get<FAuthOSSAdapter>()->GetUniqueNetId(Params.LocalUserId);
-		if (!LocalUserId->IsValid())
+		const FUniqueNetIdPtr LocalUserId = Services.Get<FAuthOSSAdapter>()->GetUniqueNetId(Params.LocalUserId);
+		if (!LocalUserId)
 		{
 			Op.SetError(Errors::InvalidUser());
 			return MakeFulfilledPromise<void>().GetFuture();
@@ -281,8 +281,8 @@ TOnlineAsyncOpHandle<FUserFileDeleteFile> FUserFileOSSAdapter::DeleteFile(FUserF
 	{
 		const FUserFileDeleteFile::Params& Params = Op.GetParams();
 
-		const FUniqueNetIdRef LocalUserId = Services.Get<FAuthOSSAdapter>()->GetUniqueNetId(Params.LocalUserId);
-		if (!LocalUserId->IsValid())
+		const FUniqueNetIdPtr LocalUserId = Services.Get<FAuthOSSAdapter>()->GetUniqueNetId(Params.LocalUserId);
+		if (!LocalUserId)
 		{
 			Op.SetError(Errors::InvalidUser());
 			return MakeFulfilledPromise<void>().GetFuture();
diff --git a/Engine/Plugins/Online/Experimental/OnlineServicesOSSAdapter/Source/Public/Online/MulticastAdapter.h b/Engine/Plugins/Online/Experimental/OnlineServicesOSSAdapter/Source/Public/Online/MulticastAdapter.h
--- a/Engine/Plugins/Online/Experimental/OnlineServicesOSSAdapter/Source/Public/Online/MulticastAdapter.h
+++ b/Engine/Plugins/Online/Experimental/OnlineServicesOSSAdapter/Source/Public/Online/MulticastAdapter.h
@@ -59,7 +59,7 @@ public:
 
 			if (bUnbind)
 			{
-				SharedAnchor.Reset();
+				// Calling remove will invalidate the closure data and destruct SharedAnchor.
 				Delegate.Remove(Handle);
 			}
 		});
diff --git a/Engine/Plugins/Online/OnlineSubsystemUtils/Source/OnlineSubsystemUtils/Private/OnlineServicesEngineInterfaceImpl.cpp b/Engine/Plugins/Online/OnlineSubsystemUtils/Source/OnlineSubsystemUtils/Private/OnlineServicesEngineInterfaceImpl.cpp
--- a/Engine/Plugins/Online/OnlineSubsystemUtils/Source/OnlineSubsystemUtils/Private/OnlineServicesEngineInterfaceImpl.cpp
+++ b/Engine/Plugins/Online/OnlineSubsystemUtils/Source/OnlineSubsystemUtils/Private/OnlineServicesEngineInterfaceImpl.cpp
@@ -108,11 +108,11 @@ FUniqueNetIdWrapper UOnlineServicesEngineInterfaceImpl::GetUniquePlayerIdWrapper
 	{
 		if (UE::Online::IAuthPtr AuthPtr = OnlineServices->GetAuthInterface())
 		{
-			UE::Online::FAuthGetAccountByPlatformUserId::Params GetAccountParams = { FPlatformMisc::GetPlatformUserForUserIndex(LocalUserNum) };
-			UE::Online::TOnlineResult<UE::Online::FAuthGetAccountByPlatformUserId> GetAccountResult = AuthPtr->GetAccountByPlatformUserId(MoveTemp(GetAccountParams));
+			UE::Online::FAuthGetLocalOnlineUserByPlatformUserId::Params GetAccountParams = { FPlatformMisc::GetPlatformUserForUserIndex(LocalUserNum) };
+			UE::Online::TOnlineResult<UE::Online::FAuthGetLocalOnlineUserByPlatformUserId> GetAccountResult = AuthPtr->GetLocalOnlineUserByPlatformUserId(MoveTemp(GetAccountParams));
 			if (GetAccountResult.IsOk())
 			{
-				return FUniqueNetIdWrapper(GetAccountResult.GetOkValue().AccountInfo->UserId);
+				return FUniqueNetIdWrapper(GetAccountResult.GetOkValue().AccountInfo->AccountId);
 			}
 		}
 	}
@@ -128,11 +128,12 @@ FString UOnlineServicesEngineInterfaceImpl::GetPlayerNickname(UWorld* World, con
 	{
 		if (UE::Online::IAuthPtr AuthPtr = OnlineServices->GetAuthInterface())
 		{
-			UE::Online::FAuthGetAccountByAccountId::Params GetAccountParams = { UniqueId.GetV2() };
-			UE::Online::TOnlineResult<UE::Online::FAuthGetAccountByAccountId> GetAccountResult = AuthPtr->GetAccountByAccountId(MoveTemp(GetAccountParams));
+			UE::Online::FAuthGetLocalOnlineUserByOnlineAccountId::Params GetAccountParams = { UniqueId.GetV2() };
+			UE::Online::TOnlineResult<UE::Online::FAuthGetLocalOnlineUserByOnlineAccountId> GetAccountResult = AuthPtr->GetLocalOnlineUserByOnlineAccountId(MoveTemp(GetAccountParams));
 			if (GetAccountResult.IsOk())
 			{
-				return GetAccountResult.GetOkValue().AccountInfo->DisplayName;
+				const UE::Online::FSchemaVariant* DisplayName = GetAccountResult.GetOkValue().AccountInfo->Attributes.Find(UE::Online::AccountAttributeData::DisplayName);
+				return DisplayName ? DisplayName->GetString() : FString();
 			}
 		}
 	}
@@ -148,11 +149,14 @@ bool UOnlineServicesEngineInterfaceImpl::GetPlayerPlatformNickname(UWorld* World
 	{
 		if (UE::Online::IAuthPtr AuthPtr = OnlineServices->GetAuthInterface())
 		{
-			UE::Online::FAuthGetAccountByPlatformUserId::Params GetAccountParams = { FPlatformMisc::GetPlatformUserForUserIndex(LocalUserNum) };
-			UE::Online::TOnlineResult<UE::Online::FAuthGetAccountByPlatformUserId> GetAccountResult = AuthPtr->GetAccountByPlatformUserId(MoveTemp(GetAccountParams));
+			UE::Online::FAuthGetLocalOnlineUserByPlatformUserId::Params GetAccountParams = { FPlatformMisc::GetPlatformUserForUserIndex(LocalUserNum) };
+			UE::Online::TOnlineResult<UE::Online::FAuthGetLocalOnlineUserByPlatformUserId> GetAccountResult = AuthPtr->GetLocalOnlineUserByPlatformUserId(MoveTemp(GetAccountParams));
 			if (GetAccountResult.IsOk())
 			{
-				OutNickname = GetAccountResult.GetOkValue().AccountInfo->DisplayName;
+				if (const UE::Online::FSchemaVariant* DisplayName = GetAccountResult.GetOkValue().AccountInfo->Attributes.Find(UE::Online::AccountAttributeData::DisplayName))
+				{
+					OutNickname = DisplayName->GetString();
+				}
 				return !OutNickname.IsEmpty();
 			}
 		}
@@ -192,8 +196,8 @@ bool UOnlineServicesEngineInterfaceImpl::IsLoggedIn(UWorld* World, int32 LocalUs
 	{
 		if (UE::Online::IAuthPtr AuthPtr = OnlineServices->GetAuthInterface())
 		{
-			UE::Online::FAuthGetAccountByPlatformUserId::Params GetAccountParams = { FPlatformMisc::GetPlatformUserForUserIndex(LocalUserNum) };
-			UE::Online::TOnlineResult<UE::Online::FAuthGetAccountByPlatformUserId> GetAccountResult = AuthPtr->GetAccountByPlatformUserId(MoveTemp(GetAccountParams));
+			UE::Online::FAuthGetLocalOnlineUserByPlatformUserId::Params GetAccountParams = { FPlatformMisc::GetPlatformUserForUserIndex(LocalUserNum) };
+			UE::Online::TOnlineResult<UE::Online::FAuthGetLocalOnlineUserByPlatformUserId> GetAccountResult = AuthPtr->GetLocalOnlineUserByPlatformUserId(MoveTemp(GetAccountParams));
 			if (GetAccountResult.IsOk())
 			{
 				return GetAccountResult.GetOkValue().AccountInfo->LoginStatus == UE::Online::ELoginStatus::LoggedIn;
diff --git a/Engine/Plugins/Online/OnlineSubsystemUtils/Source/OnlineSubsystemUtils/Private/OnlineServicesEngineUtilsImpl.cpp b/Engine/Plugins/Online/OnlineSubsystemUtils/Source/OnlineSubsystemUtils/Private/OnlineServicesEngineUtilsImpl.cpp
--- a/Engine/Plugins/Online/OnlineSubsystemUtils/Source/OnlineSubsystemUtils/Private/OnlineServicesEngineUtilsImpl.cpp
+++ b/Engine/Plugins/Online/OnlineSubsystemUtils/Source/OnlineSubsystemUtils/Private/OnlineServicesEngineUtilsImpl.cpp
@@ -149,7 +149,7 @@ void FOnlineServicesEngineUtils::GetPIELogins(TArray<FAuthLogin::Params>& Logins
 		UE::Online::FAuthLogin::Params LoginParameters;
 		LoginParameters.CredentialsId = PIELoginCredentials.Id;
 		LoginParameters.CredentialsToken.Set<FString>(PIELoginCredentials.Token);
-		LoginParameters.CredentialsType = PIELoginCredentials.Type;
+		LoginParameters.CredentialsType = *PIELoginCredentials.Type;
 		Logins.Emplace(MoveTemp(LoginParameters));
 	}
 }
