OnlineExecHandler major feature additions and reshuffle
* Added support for TSharedRef/TSharedPtr (no special syntax, just uses the syntax of whatevers in the ptr), TVariant (prefix with the type, i.e. s: for string, i: for int, so on), TOptional ("null" to unset the optional)
* Objects, arrays, and maps support a json-esque syntax using brackets {} []. Arrays for example can be declared as [1, 2, 3, 4]. Objects must still be specified in arg order, e.g. {1 true s:hello false}
* FOnlineAccountIdHandle now supports passing in a single digit 0-9 which will automatically grab that logged in user and use that user's ID
* Better debugging messages when fields fail to parse
* Basic testing support for some of the parsing

Still todo (backlogged temporarily to make sure the ecom interface gets done for 5.1)
* Saveable object types (allow adding some object via config or "SaveObject [REDACTED_HANDLE] ..." that can be then referenced instead of having to re-type the entire logic every time
* Load an FOnlineAccountIdHandle from its platform string instead of its handle (this needs a bit of refactoring on the net id registry side)

diff --git a/Engine/Plugins/Online/Experimental/OnlineServices/Source/OnlineServicesCommon/Private/Online/OnlineExecHandler.cpp b/Engine/Plugins/Online/Experimental/OnlineServices/Source/OnlineServicesCommon/Private/Online/OnlineExecHandler.cpp
new file mode 100644
--- /dev/null
+++ b/Engine/Plugins/Online/Experimental/OnlineServices/Source/OnlineServicesCommon/Private/Online/OnlineExecHandler.cpp
@@ -0,0 +1,8 @@
+// Copyright Epic Games, Inc. All Rights Reserved.
+#include "Online/OnlineExecHandler.h"
+
+namespace UE::Online::Private {
+
+TSharedPtr<IOnlineServices> GExecCmdServices;
+
+} // namespace UE::Online::Private
\ No newline at end of file
diff --git a/Engine/Plugins/Online/Experimental/OnlineServices/Source/OnlineServicesCommon/Private/Online/SessionsCommon.cpp b/Engine/Plugins/Online/Experimental/OnlineServices/Source/OnlineServicesCommon/Private/Online/SessionsCommon.cpp
--- a/Engine/Plugins/Online/Experimental/OnlineServices/Source/OnlineServicesCommon/Private/Online/SessionsCommon.cpp
+++ b/Engine/Plugins/Online/Experimental/OnlineServices/Source/OnlineServicesCommon/Private/Online/SessionsCommon.cpp
@@ -30,9 +30,9 @@ namespace UE::Online {
 		RegisterCommand(&FSessionsCommon::LeaveSession);
 		RegisterCommand(&FSessionsCommon::FindSessions);
 		RegisterCommand(&FSessionsCommon::StartMatchmaking);
-		RegisterCommand(&FSessionsCommon::JoinSession);
+		//RegisterCommand(&FSessionsCommon::JoinSession);
 		RegisterCommand(&FSessionsCommon::SendSessionInvite);
-		RegisterCommand(&FSessionsCommon::RejectSessionInvite);
+		//RegisterCommand(&FSessionsCommon::RejectSessionInvite);
 		RegisterCommand(&FSessionsCommon::RegisterPlayers);
 		RegisterCommand(&FSessionsCommon::UnregisterPlayers);
 	}
diff --git a/Engine/Plugins/Online/Experimental/OnlineServices/Source/OnlineServicesCommon/Public/Online/OnlineExecHandler.h b/Engine/Plugins/Online/Experimental/OnlineServices/Source/OnlineServicesCommon/Public/Online/OnlineExecHandler.h
--- a/Engine/Plugins/Online/Experimental/OnlineServices/Source/OnlineServicesCommon/Public/Online/OnlineExecHandler.h
+++ b/Engine/Plugins/Online/Experimental/OnlineServices/Source/OnlineServicesCommon/Public/Online/OnlineExecHandler.h
@@ -5,9 +5,56 @@
 #include "Online/CoreOnline.h"
 #include "Online/OnlineUtils.h"
 #include "Online/OnlineServicesLog.h"
+#include "Online/OnlineResult.h"
+#include "Online/OnlineAsyncOp.h"
+
+#include "Online/OnlineServicesRegistry.h"
+#include "Online/Auth.h"
 
 #include "Templates/Invoke.h"
 #include "Templates/Models.h"
+#include "Templates/Tuple.h"
+
+/*
+	This file facilitates online command execution.
+	Online commands must be prefixed with OnlineServices, and Index=# where the # is which service you want in context
+		- The list of numbers can be obtained with OnlineServices List
+	Next argument is the interface- i.e. "Auth", "Presence"
+	Next argument is the name of the function being called, i.e. "QueryPresence"
+	Arguments after that are the function's params in declared order. Special types are as follows:
+		for FOnlineAccountIdHandle, either pass in 0-9 to grab the nth local user, or pass in the output from a handle's tostring that looks like "epic:3"
+		for TSharedPtr, simply type in the parameter as normal
+			- TSharedRef *technically* is implemented however will crash because you cannot default construct shared refs. At the moment, we recommend all console commands with shared ptr parameters to use tsharedptr instead of tsharedref.
+		for TOptional, pass in "null" and the optional will be unset
+		for TVariant, the syntax is type:value. (example: i:350) Currently implemented types below:
+			- s: for string
+			- i: for int(64)
+			- i32 for int32
+			- b for bool
+			- s for string
+			- u for user (FOnlineAccountIdHandle)
+			- d for double
+			- f for float
+			- e for enum
+				- note- if you have a custom enum, you must use the macro MAKE_VARIANT_ENUM_INFO(YourEnumName) for it to work properly
+				- if you have more than one enum you must qualify the entire enum name. using e: with a variant with more than one enum is undefined behavior
+					- e.g. "ELoginStatus:Unknown"
+		for objects (with online metadata), you can use the {} syntax and declare all parameters flatly
+			- i.e. {true "AuthLogin" s:username s:password}
+			- objects support internal objects/arrays, e.g. {{5 3} {4 false}}
+		for TArray, you can use [] syntax
+			- "[5, 3, 7, 9]"
+			- commas are optional- [5 3 7 9] is also valid
+			- arrays support internal objects/recursive arrays, e.g. [[2,3], [4,5]]
+		for TMap, you can use {} syntax. 
+			- [licenseData=licenseDataText, kNumUsers=5}
+			- please note: TMap does not currently support rich syntax so nested objects/arrays/etc may not parse properly
+
+	Full command examples:
+		OnlineServices Index=0 Auth Login Null username s:pwd []
+		OnlineServices Index=0 Presence UpdatePresence 0 {0 Unknown Unknown Unknown Test1 Test2 {a=b}}
+		OnlineServices Index=0 Pressence QueryPresence 0 0 true
+*/
 
 class UWorld;
 class FOutputDevice;
@@ -60,26 +107,158 @@ struct TOnlineInterfaceOperationMemberFunctionPtrTraits<TOnlineAsyncOpHandle<TOp
 };
 
 // Forward declarations for dependent type resolution
-inline bool ParseOnlineExecParams(const TCHAR*& Cmd, FString& Value);
-inline bool ParseOnlineExecParams(const TCHAR*& Cmd, FName& Value);
-inline bool ParseOnlineExecParams(const TCHAR*& Cmd, uint8& Value);
-inline bool ParseOnlineExecParams(const TCHAR*& Cmd, int32& Value);
-inline bool ParseOnlineExecParams(const TCHAR*& Cmd, uint32& Value);
-inline bool ParseOnlineExecParams(const TCHAR*& Cmd, int64& Value);
-inline bool ParseOnlineExecParams(const TCHAR*& Cmd, uint64& Value);
-inline bool ParseOnlineExecParams(const TCHAR*& Cmd, bool& Value);
-template <typename T> inline bool ParseOnlineExecParams(const TCHAR*& Cmd, TArray<T>& Array);
-template <typename T, typename U> inline bool ParseOnlineExecParams(const TCHAR*& Cmd, TMap<T, U>& Map);
-template <typename T> inline bool ParseOnlineExecParams(const TCHAR*& Cmd, TSet<T>& Set);
-template <typename T> inline bool ParseOnlineExecParams(const TCHAR*& Cmd, TSharedPtr<T>& Ptr);
-template <typename T> inline bool ParseOnlineExecParams(const TCHAR*& Cmd, TSharedRef<T>& Ref);
-template <typename T> inline bool ParseOnlineExecParams(const TCHAR*& Cmd, TOptional<T>& Optional);
-template <typename... Ts> inline bool ParseOnlineExecParams(const TCHAR*& Cmd, TVariant<Ts...>& Variant);
-template <typename IdType> inline bool ParseOnlineExecParams(const TCHAR*& Cmd, TOnlineIdHandle<IdType>& Value);
-template <typename T> std::enable_if_t<!TModels<Meta::COnlineMetadataAvailable, T>::Value, bool> ParseOnlineExecParams(const TCHAR*& Cmd, T& Value);
-template <typename T> std::enable_if_t<TModels<Meta::COnlineMetadataAvailable, T>::Value, bool> ParseOnlineExecParams(const TCHAR*& Cmd, T& Value);
-
-inline bool ParseOnlineExecParams(const TCHAR*& Cmd, FString& Value)
+inline bool ParseOnlineExecParams(const TCHAR*& Cmd, FString& Value, IOnlineServices* Services = nullptr);
+inline bool ParseOnlineExecParams(const TCHAR*& Cmd, FName& Value, IOnlineServices* Services = nullptr);
+inline bool ParseOnlineExecParams(const TCHAR*& Cmd, uint8& Value, IOnlineServices* Services = nullptr);
+inline bool ParseOnlineExecParams(const TCHAR*& Cmd, int32& Value, IOnlineServices* Services = nullptr);
+inline bool ParseOnlineExecParams(const TCHAR*& Cmd, uint32& Value, IOnlineServices* Services = nullptr);
+inline bool ParseOnlineExecParams(const TCHAR*& Cmd, int64& Value, IOnlineServices* Services = nullptr);
+inline bool ParseOnlineExecParams(const TCHAR*& Cmd, uint64& Value, IOnlineServices* Services = nullptr);
+inline bool ParseOnlineExecParams(const TCHAR*& Cmd, bool& Value, IOnlineServices* Services = nullptr);
+template <typename T> inline bool ParseOnlineExecParams(const TCHAR*& Cmd, TArray<T>& Array, IOnlineServices* Services = nullptr);
+template <typename T, typename U> inline bool ParseOnlineExecParams(const TCHAR*& Cmd, TMap<T, U>& Map, IOnlineServices* Services = nullptr);
+template <typename T> inline bool ParseOnlineExecParams(const TCHAR*& Cmd, TSet<T>& Set, IOnlineServices* Services = nullptr);
+template <typename T> inline bool ParseOnlineExecParams(const TCHAR*& Cmd, TSharedPtr<T>& Ptr, IOnlineServices* Services = nullptr);
+template <typename T> inline bool ParseOnlineExecParams(const TCHAR*& Cmd, TSharedRef<T>& Ref, IOnlineServices* Services = nullptr);
+template <typename T> inline bool ParseOnlineExecParams(const TCHAR*& Cmd, TOptional<T>& Optional, IOnlineServices* Services = nullptr);
+template <typename... Ts> inline bool ParseOnlineExecParams(const TCHAR*& Cmd, TVariant<Ts...>& Variant, IOnlineServices* Services = nullptr);
+template <typename IdType> inline bool ParseOnlineExecParams(const TCHAR*& Cmd, TOnlineIdHandle<IdType>& Value, IOnlineServices* Services = nullptr);
+template <typename T> std::enable_if_t<!TModels<Meta::COnlineMetadataAvailable, T>::Value, bool> ParseOnlineExecParams(const TCHAR*& Cmd, T& Value, IOnlineServices* Services = nullptr);
+template <typename T> std::enable_if_t<TModels<Meta::COnlineMetadataAvailable, T>::Value, bool> ParseOnlineExecParams(const TCHAR*& Cmd, T& Value, IOnlineServices* Services = nullptr);
+
+template<typename T, typename VariantObject>
+struct TOnlineVariantVisitInfo
+{
+	// todo: how to get this to work without it triggering on every class?
+	//static_assert(sizeof(T) == -1, "If you are hitting this, you are using a TVariant with a value that does not have a TOnlineVariantVisitInfo! Make sure you are adding a TOnlineVariantVisitInfo for your class! (Check LastType or FirstType below for the type that needs attention)");
+
+public:
+	inline static TArray<FString> Prefixes = { TEXT("[GARBAGE]") };
+
+	static bool Assign(FString& VariantValue, VariantObject& Variant, IOnlineServices* Services)
+	{
+		return false;
+	}
+};
+
+template< typename VariantObject >
+struct TOnlineVariantVisitInfo<int64, VariantObject>
+{
+public:
+	inline static TArray<FString> Prefixes = { TEXT("i"),TEXT("i64"),TEXT("int"),TEXT("int64") };
+
+	static bool Assign(FString& VariantValue, VariantObject& Variant, IOnlineServices* Services)
+	{
+		Variant.template Set<int64>(FCString::Atoi64(*VariantValue));
+		return true;
+	}
+};
+
+template< typename VariantObject >
+struct TOnlineVariantVisitInfo<int32, VariantObject>
+{
+public:
+	inline static TArray<FString> Prefixes = { TEXT("i32"),TEXT("int32") };
+
+	static bool Assign(FString& VariantValue, VariantObject& Variant, IOnlineServices* Services)
+	{
+		Variant.template Set<int32>(FCString::Atoi(*VariantValue));
+		return true;
+	}
+};
+
+
+template< typename VariantObject >
+struct TOnlineVariantVisitInfo<double, VariantObject>
+{
+public:
+	inline static TArray<FString> Prefixes = { TEXT("d"),TEXT("double") };
+
+	static bool Assign(FString& VariantValue, VariantObject& Variant, IOnlineServices* Services)
+	{
+		Variant.template Set<double>(FCString::Atod(*VariantValue));
+		return true;
+	}
+};
+
+template< typename VariantObject >
+struct TOnlineVariantVisitInfo<float, VariantObject>
+{
+public:
+	inline static TArray<FString> Prefixes = { TEXT("f"),TEXT("float") };
+
+	static bool Assign(FString& VariantValue, VariantObject& Variant, IOnlineServices* Services)
+	{
+		Variant.template Set<float>(FCString::Atof(*VariantValue));
+		return true;
+	}
+};
+
+template< typename VariantObject >
+struct TOnlineVariantVisitInfo<bool, VariantObject>
+{
+public:
+	inline static TArray<FString> Prefixes = { TEXT("b"),TEXT("bool"),TEXT("boolean") };
+
+	static bool Assign(FString& VariantValue, VariantObject& Variant, IOnlineServices* Services)
+	{
+		Variant.template Set<bool>(FCString::ToBool(*VariantValue));
+		return true;
+	}
+};
+
+
+template< typename VariantObject >
+struct TOnlineVariantVisitInfo<FString, VariantObject>
+{
+public:
+	inline static TArray<FString> Prefixes = { TEXT("s"),TEXT("str"),TEXT("string") };
+
+	static bool Assign(FString& VariantValue, VariantObject& Variant, IOnlineServices* Services)
+	{
+		Variant.template Set<FString>(VariantValue);
+		return true;
+	}
+};
+
+template< typename VariantObject >
+struct TOnlineVariantVisitInfo<FOnlineAccountIdHandle, VariantObject>
+{
+public:
+	inline static TArray<FString> Prefixes = { TEXT("u"),TEXT("user") };
+
+	static bool Assign(FString& VariantValue, VariantObject& Variant, IOnlineServices* Services)
+	{
+		FOnlineAccountIdHandle Handle;
+		const TCHAR* ValueTCHAR = *VariantValue;
+		if (ParseOnlineExecParams<TOnlineIdHandle<OnlineIdHandleTags::FAccount>>(ValueTCHAR, Handle, Services))
+		{
+			Variant.template Set<FOnlineAccountIdHandle>(Handle);
+			return true;
+		}
+		return false;
+	}
+};
+
+
+#define MAKE_VARIANT_ENUM_INFO_SHORTNAME(EnumName, ShortName) template< typename VariantObject >\
+struct UE::Online::Private::TOnlineVariantVisitInfo<EnumName, VariantObject>\
+{\
+public:\
+	inline static TArray<FString> Prefixes = {TEXT(#ShortName),TEXT("e")};\
+	static bool Assign(FString& VariantValue, VariantObject& Variant, IOnlineServices* Services)\
+	{\
+		EnumName Status;\
+		LexFromString(Status, *VariantValue);\
+		Variant.template Set<EnumName>(Status);\
+		return true;\
+	}\
+};
+
+#define MAKE_VARIANT_ENUM_INFO(EnumName) MAKE_VARIANT_ENUM_INFO_SHORTNAME(EnumName, EnumName);
+
+
+inline bool ParseOnlineExecParams(const TCHAR*& Cmd, FString& Value, IOnlineServices* Services)
 {
 	if (!FParse::Token(Cmd, Value, true))
 	{
@@ -89,7 +268,7 @@ inline bool ParseOnlineExecParams(const TCHAR*& Cmd, FString& Value)
 	return true;
 }
 
-inline bool ParseOnlineExecParams(const TCHAR*& Cmd, FName& Value)
+inline bool ParseOnlineExecParams(const TCHAR*& Cmd, FName& Value, IOnlineServices* Services)
 {
 	FString StringValue;
 	if (!ParseOnlineExecParams(Cmd, StringValue))
@@ -101,7 +280,7 @@ inline bool ParseOnlineExecParams(const TCHAR*& Cmd, FName& Value)
 	return true;
 }
 
-inline bool ParseOnlineExecParams(const TCHAR*& Cmd, uint8& Value)
+inline bool ParseOnlineExecParams(const TCHAR*& Cmd, uint8& Value, IOnlineServices* Services)
 {
 	FString Token;
 	if (FParse::Token(Cmd, Token, true))
@@ -116,7 +295,7 @@ inline bool ParseOnlineExecParams(const TCHAR*& Cmd, uint8& Value)
 	return true;
 }
 
-inline bool ParseOnlineExecParams(const TCHAR*& Cmd, int32& Value)
+inline bool ParseOnlineExecParams(const TCHAR*& Cmd, int32& Value, IOnlineServices* Services)
 {
 	FString Token;
 	if (FParse::Token(Cmd, Token, true))
@@ -131,7 +310,7 @@ inline bool ParseOnlineExecParams(const TCHAR*& Cmd, int32& Value)
 	return true;
 }
 
-inline bool ParseOnlineExecParams(const TCHAR*& Cmd, uint32& Value)
+inline bool ParseOnlineExecParams(const TCHAR*& Cmd, uint32& Value, IOnlineServices* Services)
 {
 	FString Token;
 	if (FParse::Token(Cmd, Token, true))
@@ -146,7 +325,7 @@ inline bool ParseOnlineExecParams(const TCHAR*& Cmd, uint32& Value)
 	return true;
 }
 
-inline bool ParseOnlineExecParams(const TCHAR*& Cmd, int64& Value)
+inline bool ParseOnlineExecParams(const TCHAR*& Cmd, int64& Value, IOnlineServices* Services)
 {
 	FString Token;
 	if (FParse::Token(Cmd, Token, true))
@@ -161,7 +340,7 @@ inline bool ParseOnlineExecParams(const TCHAR*& Cmd, int64& Value)
 	return true;
 }
 
-inline bool ParseOnlineExecParams(const TCHAR*& Cmd, uint64& Value)
+inline bool ParseOnlineExecParams(const TCHAR*& Cmd, uint64& Value, IOnlineServices* Services)
 {
 	FString Token;
 	if (FParse::Token(Cmd, Token, true))
@@ -176,7 +355,7 @@ inline bool ParseOnlineExecParams(const TCHAR*& Cmd, uint64& Value)
 	return true;
 }
 
-inline bool ParseOnlineExecParams(const TCHAR*& Cmd, bool& Value)
+inline bool ParseOnlineExecParams(const TCHAR*& Cmd, bool& Value, IOnlineServices* Services)
 {
 	FString Token;
 	if (FParse::Token(Cmd, Token, true))
@@ -191,7 +370,7 @@ inline bool ParseOnlineExecParams(const TCHAR*& Cmd, bool& Value)
 	return true;
 }
 
-inline bool ParseOnlineExecParams(const TCHAR*& Cmd, FPlatformUserId& Value)
+inline bool ParseOnlineExecParams(const TCHAR*& Cmd, FPlatformUserId& Value, IOnlineServices* Services)
 {
 	FString Token;
 	if (FParse::Token(Cmd, Token, true))
@@ -206,20 +385,188 @@ inline bool ParseOnlineExecParams(const TCHAR*& Cmd, FPlatformUserId& Value)
 	return true;
 }
 
+inline TCHAR PeekNextChar(const TCHAR* Cmd)
+{
+	while (FChar::IsWhitespace(*Cmd))
+	{
+		Cmd++;
+	}
+	return *Cmd;
+}
+
+// taken from FParse::Token, functionality added
+// (possibly can be merged back?)
+inline bool ParseTokenExt(const TCHAR*& Str, FString& Arg, bool UseEscape, bool bCommaSeparated = false, bool bIncludeDelimiters = false)
+{
+	Arg.Empty();
+
+	// Skip preceeding spaces and tabs. (edit: also optional comma removal)
+	while (FChar::IsWhitespace(*Str) || (bCommaSeparated && *Str == TEXT(',')))
+	{
+		Str++;
+	}
+
+	// EDITS Begin
+	bool bHasEnteredObject = false;
+	// Support for defining a nested object-like syntax, e,g, {Cats Dogs {Tigers Bears}}- ParseTokenExt will return "Cats Dogs {Tigers Bears}"
+	if (*Str == TEXT('{') || *Str == TEXT('['))
+	{
+		bHasEnteredObject = true;
+		TCHAR Incrementor = (*Str == TEXT('{')) ? TEXT('{') : TEXT('[');
+		TCHAR Decrementor = (*Str == TEXT('{')) ? TEXT('}') : TEXT(']');
+		bool bInQuote = false;
+		int Layer = 1;
+		if (bIncludeDelimiters)
+		{
+			Arg += *Str;
+		}
+		Str++; 
+
+		while (1)
+		{
+			TCHAR Character = *Str;
+			if (Character == 0)
+			{
+				break;
+			}
+			Str++;
+
+			// Preserve escapes if they're in a quoted string (the check for " is in the else to let \" work as expected)
+			if (Character == TEXT('\\') && UseEscape && bInQuote)
+			{
+				Arg += Character;
+
+				Character = *Str;
+				if (!Character)
+				{
+					break;
+				}
+				Str++;
+			}
+			else if (Character == TEXT('"'))
+			{
+				bInQuote = !bInQuote;
+			}
+			else if (!bInQuote)
+			{
+				if (Character == Decrementor)
+				{
+					Layer--;
+				}
+				else if (Character == Incrementor)
+				{
+					Layer++;
+				}
+				if (Layer == 0)
+				{
+					if (bIncludeDelimiters)
+					{
+						Arg += Character;
+					}
+					break;
+				}
+			}
+			Arg += Character;
+		}
+	}
+	else
+	// Edits End
+	if (*Str == TEXT('"'))
+	{
+		// Get quoted string.
+		if (bIncludeDelimiters)
+		{
+			Arg += *Str;
+		}
+		Str++;
+		while (*Str && *Str != TCHAR('"'))
+		{
+			TCHAR c = *Str++;
+			if (c == TEXT('\\') && UseEscape)
+			{
+				// Get escape.
+				c = *Str++;
+				if (!c)
+				{
+					break;
+				}
+			}
+
+			Arg += c;
+		}
+
+		if (*Str == TEXT('"'))
+		{
+			if (bIncludeDelimiters)
+			{
+				Arg += *Str;
+			}
+			Str++;
+		}
+	}
+	else
+	{
+		// Get unquoted string (that might contain a quoted part, which will be left intact).
+		// For example, -ARG="foo bar baz", will be treated as one token, with quotes intact
+		bool bInQuote = false;
+
+		while (1)
+		{
+			TCHAR Character = *Str;
+			if ((Character == 0) || (FChar::IsWhitespace(Character) && !bInQuote))
+			{
+				break;
+			}
+			Str++;
+
+			// Preserve escapes if they're in a quoted string (the check for " is in the else to let \" work as expected)
+			if (Character == TEXT('\\') && UseEscape && bInQuote)
+			{
+				Arg += Character;
+
+				Character = *Str;
+				if (!Character)
+				{
+					break;
+				}
+				Str++;
+			}
+			else if (Character == TEXT('"'))
+			{
+				bInQuote = !bInQuote;
+			}
+			// EDITS begin- if we're doing comma-separated-values, break if we see a comma
+			if (bCommaSeparated)
+			{
+				if (!bInQuote && Character == TEXT(','))
+				{
+					break;
+				}
+			}
+			// EDITS end
+
+			Arg += Character;
+		}
+	}
+
+	// bHasEnteredObject allows us to declare empty arrays/objects like [] that are prefectly valid even though they are empty...
+	return Arg.Len() > 0 || bHasEnteredObject;
+}
+
 template <typename T>
-inline bool ParseOnlineExecParams(const TCHAR*& Cmd, TArray<T>& Array)
+inline bool ParseOnlineExecParams(const TCHAR*& Cmd, TArray<T>& Array, IOnlineServices* Services)
 {
 	FString Token;
-	if (FParse::Token(Cmd, Token, true))
+	if (ParseTokenExt(Cmd, Token, true))
 	{
-		TArray<FString> TokenArray;
-		Token.ParseIntoArray(TokenArray, TEXT(","));
-		Array.Reserve(TokenArray.Num());
-		for (const FString& ArrayToken : TokenArray)
+		FString NextToken;
+		const TCHAR* TokenTCHAR = *Token;
+		while (ParseTokenExt(TokenTCHAR, NextToken, true, true, true))
 		{
 			T Value;
-			const TCHAR* ArrayTokenTCHAR = *ArrayToken;
-			if (ParseOnlineExecParams(ArrayTokenTCHAR, Value))
+			FString Trimmed = NextToken.TrimStartAndEnd();
+			const TCHAR* ArrayTokenTCHAR = *Trimmed;
+			if (ParseOnlineExecParams(ArrayTokenTCHAR, Value, Services))
 			{
 				Array.Emplace(MoveTempIfPossible(Value));
 			}
@@ -230,13 +577,14 @@ inline bool ParseOnlineExecParams(const TCHAR*& Cmd, TArray<T>& Array)
 		return false;
 	}
 	return true;
+
 }
 
 template <typename T, typename U>
-inline bool ParseOnlineExecParams(const TCHAR*& Cmd, TMap<T, U>& Map)
+inline bool ParseOnlineExecParams(const TCHAR*& Cmd, TMap<T, U>& Map, IOnlineServices* Services)
 {
 	FString Token;
-	if (FParse::Token(Cmd, Token, true))
+	if (ParseTokenExt(Cmd, Token, true))
 	{
 		TArray<FString> TokenArray;
 		Token.ParseIntoArray(TokenArray, TEXT(","));
@@ -244,25 +592,40 @@ inline bool ParseOnlineExecParams(const TCHAR*& Cmd, TMap<T, U>& Map)
 		for (const FString& ArrayToken : TokenArray)
 		{
 			TArray<FString> TokenValuePair;
-			Token.ParseIntoArray(TokenValuePair, TEXT("="));
+			ArrayToken.ParseIntoArray(TokenValuePair, TEXT("="));
 
-			if (TokenValuePair.Num() > 1)
+			if (TokenValuePair.Num() > 2)
 			{
 				T Key;
 				U Value;
 
-				const TCHAR* ArrayTokenKeyTCHAR = *TokenValuePair[0];
-				if (ParseOnlineExecParams(ArrayTokenKeyTCHAR, Key))
+				FString TrimmedKey = TokenValuePair[0].TrimStartAndEnd();
+				const TCHAR* ArrayTokenKeyTCHAR = *TrimmedKey;
+				if (ParseOnlineExecParams(ArrayTokenKeyTCHAR, Key, Services))
 				{
-					if (TokenValuePair.Num() > 2)
+					FString TrimmedValue = TokenValuePair[2].TrimStartAndEnd();
+					const TCHAR* ArrayTokenValueTCHAR = *TrimmedValue;
+					if (ParseOnlineExecParams(ArrayTokenValueTCHAR, Value, Services))
 					{
-						const TCHAR* ArrayTokenValueTCHAR = *TokenValuePair[1];
-						ParseOnlineExecParams(ArrayTokenValueTCHAR, Value);
-
 						Map.Emplace(MoveTempIfPossible(Key), MoveTempIfPossible(Value));
 					}
+					else
+					{
+						UE_LOG(LogOnlineServices, Warning, TEXT("Failed to parse exec param value for map: %s"), *TrimmedValue);
+						return false;
+					}
+				}
+				else
+				{
+					UE_LOG(LogOnlineServices, Warning, TEXT("Failed to parse exec param key for map: %s"), *TrimmedKey);
+					return false;
 				}
 			}
+			else
+			{
+				UE_LOG(LogOnlineServices, Warning, TEXT("Received map key/value with invalid number of [=]: %s"), *ArrayToken);
+				return false;
+			}
 		}
 	}
 	else
@@ -273,7 +636,7 @@ inline bool ParseOnlineExecParams(const TCHAR*& Cmd, TMap<T, U>& Map)
 }
 
 template <typename T>
-inline bool ParseOnlineExecParams(const TCHAR*& Cmd, TSet<T>& Set)
+inline bool ParseOnlineExecParams(const TCHAR*& Cmd, TSet<T>& Set, IOnlineServices* Services)
 {
 	FString Token;
 	if (FParse::Token(Cmd, Token, true))
@@ -300,35 +663,128 @@ inline bool ParseOnlineExecParams(const TCHAR*& Cmd, TSet<T>& Set)
 }
 
 template <typename T>
-inline bool ParseOnlineExecParams(const TCHAR*& Cmd, TSharedPtr<T>& Ptr)
+inline bool ParseOnlineExecParams(const TCHAR*& Cmd, TSharedPtr<T>& Ptr, IOnlineServices* Services)
 {
-	// TODO: This is a temporary stub implementation
-	return true;
+	if (Ptr.IsValid())
+	{
+		return ParseOnlineExecParams(Cmd, *Ptr, Services);
+	}
+	else 
+	{
+		T* RawPtr = new T();
+		bool bSuccess = ParseOnlineExecParams(Cmd, *RawPtr, Services);
+		if (bSuccess)
+		{
+			Ptr = MakeShareable<T>(RawPtr);
+		}
+		else
+		{
+			delete RawPtr;
+		}
+		return bSuccess;
+	}
 }
 
 template <typename T>
-inline bool ParseOnlineExecParams(const TCHAR*& Cmd, TSharedRef<T>& Ref)
+inline bool ParseOnlineExecParams(const TCHAR*& Cmd, TSharedRef<T>& Ref, IOnlineServices* Services)
 {
-	// TODO: This is a temporary stub implementation
-	return true;
+	return ParseOnlineExecParams(Cmd, *Ref, Services);
 }
 
 template <typename T>
-inline bool ParseOnlineExecParams(const TCHAR*& Cmd, TOptional<T>& Optional)
+inline bool ParseOnlineExecParams(const TCHAR*& Cmd, TOptional<T>& Optional, IOnlineServices* Services)
 {
-	// TODO: This is a temporary stub implementation
-	return true;
+	FString PrecedingToken;
+	const TCHAR* OldCmd = Cmd;
+	FParse::Token(Cmd, PrecedingToken, true);
+
+	if (PrecedingToken.Equals(TEXT("null"))) 
+	{
+		Optional.Reset();
+		return true;
+	}
+
+	// go back to before the FParse call because we need that value again (no way to "peek" a token sadly)
+	Cmd = OldCmd;
+
+	T Temp;
+	bool bSuccess = ParseOnlineExecParams(Cmd, Temp, Services);
+	if (bSuccess)
+	{
+		Optional = MoveTemp(Temp);
+	}
+	
+	return bSuccess;
+}
+
+template<typename VariantType, typename LastType>
+static bool VisitVariantTypes(VariantType& Variant, FString& VariantPrefix, FString& VariantValue, IOnlineServices* Services)
+{
+	for(const FString& Str : TOnlineVariantVisitInfo<LastType, VariantType>::Prefixes)
+	{
+		if (Str.Equals(VariantPrefix, ESearchCase::IgnoreCase))
+		{
+			return TOnlineVariantVisitInfo<LastType, VariantType>::Assign(VariantValue, Variant, Services);
+		}
+	}
+	return false;
+}
+
+template<typename VariantType, typename FirstType, typename SecondType, typename... RemainingTypes>
+inline bool VisitVariantTypes(VariantType& Variant, FString& VariantPrefix, FString& VariantValue, IOnlineServices* Services)
+{
+	for(const FString& Str : TOnlineVariantVisitInfo<FirstType, VariantType>::Prefixes)
+	{
+		if (Str.Equals(VariantPrefix, ESearchCase::IgnoreCase))
+		{
+			return TOnlineVariantVisitInfo<FirstType, VariantType>::Assign(VariantValue, Variant, Services);
+		}
+	}
+	return VisitVariantTypes<VariantType, SecondType, RemainingTypes...>(Variant, VariantPrefix, VariantValue, Services);
 }
 
+
 template <typename... Ts>
-inline bool ParseOnlineExecParams(const TCHAR*& Cmd, TVariant<Ts...>& Variant)
+inline bool ParseOnlineExecParams(const TCHAR*& Cmd, TVariant<Ts...>& Variant, IOnlineServices* Services)
 {
-	// TODO: This is a temporary stub implementation
-	return true;
+	FString Token;
+	if (FParse::Token(Cmd, Token, true))
+	{
+		FString VariantType;
+		FString VariantValue;
+
+		if (Token.Split(TEXT(":"), &VariantType, &VariantValue, ESearchCase::CaseSensitive))
+		{
+			return VisitVariantTypes<TVariant<Ts...>, Ts...>(Variant, VariantType, VariantValue, Services);
+		}
+	}
+	return false;
 }
 
 template<typename IdType>
-inline bool ParseOnlineExecParams(const TCHAR*& Cmd, TOnlineIdHandle<IdType>& Value)
+inline bool ParseOnlineExecParams(const TCHAR*& Cmd, TOnlineIdHandle<IdType>& Value, IOnlineServices* Services)
+{
+	FString Token;
+	if (FParse::Token(Cmd, Token, true))
+	{
+		FString ServicesString;
+		FString HandleString;
+
+		if (Token.Split(TEXT(":"), &ServicesString, &HandleString, ESearchCase::CaseSensitive))
+		{
+			EOnlineServices ServiceEnum;
+			LexFromString(ServiceEnum, *ServicesString);
+			uint32 Handle = static_cast<uint32>(FCString::Strtoui64(*HandleString, nullptr, 10));
+			Value = TOnlineIdHandle<IdType>(ServiceEnum, Handle);
+			return true;
+		}
+	}
+
+	return false;
+}
+
+template<>
+inline bool ParseOnlineExecParams<OnlineIdHandleTags::FAccount>(const TCHAR*& Cmd, TOnlineIdHandle<OnlineIdHandleTags::FAccount>& Value, IOnlineServices* Services)
 {
 	FString Token;
 	if (FParse::Token(Cmd, Token, true))
@@ -338,19 +794,43 @@ inline bool ParseOnlineExecParams(const TCHAR*& Cmd, TOnlineIdHandle<IdType>& Va
 
 		if (Token.Split(TEXT(":"), &ServicesString, &HandleString, ESearchCase::CaseSensitive))
 		{
-			EOnlineServices Services;
-			LexFromString(Services, *ServicesString);
+			EOnlineServices ServicesEnum;
+			LexFromString(ServicesEnum, *ServicesString);
 			uint32 Handle = static_cast<uint32>(FCString::Strtoui64(*HandleString, nullptr, 10));
-			Value = TOnlineIdHandle<IdType>(Services, Handle);
+			Value = TOnlineIdHandle<OnlineIdHandleTags::FAccount>(ServicesEnum, Handle);
 			return true;
 		}
+		else if (Token.Len() == 1 && **Token >= TEXT('0') && **Token <= TEXT('9'))
+		{
+			FPlatformUserId PlatformUserId = FPlatformMisc::GetPlatformUserForUserIndex(FCString::Strtoi(*Token, nullptr, 10));
+			
+			// may need to move this to its own .cpp file to avoid circular dependencies
+
+			if (Services)
+			{
+				IAuthPtr AuthPtr = Services->GetAuthInterface();
+				if (AuthPtr.IsValid())
+				{
+					TOnlineResult<FAuthGetLocalOnlineUserByPlatformUserId> AuthResult = AuthPtr->GetLocalOnlineUserByPlatformUserId({PlatformUserId});
+					if (AuthResult.IsOk())
+					{
+						Value = AuthResult.GetOkValue().AccountInfo->AccountId;
+						return true;
+					}
+					else
+					{
+						UE_LOG(LogOnlineServices, Warning, TEXT("Failed to resolve user id %d - not logged in!"), PlatformUserId.GetInternalId());
+					}
+				}
+			}
+		}
 	}
 
 	return false;
 }
 
 template <typename T>
-std::enable_if_t<!TModels<Meta::COnlineMetadataAvailable, T>::Value, bool> ParseOnlineExecParams(const TCHAR*& Cmd, T& Value)
+std::enable_if_t<!TModels<Meta::COnlineMetadataAvailable, T>::Value, bool> ParseOnlineExecParams(const TCHAR*& Cmd, T& Value, IOnlineServices* Services)
 {
 	FString Token;
 	if (FParse::Token(Cmd, Token, true))
@@ -366,13 +846,28 @@ std::enable_if_t<!TModels<Meta::COnlineMetadataAvailable, T>::Value, bool> Parse
 }
 
 template <typename T>
-std::enable_if_t<TModels<Meta::COnlineMetadataAvailable, T>::Value, bool> ParseOnlineExecParams(const TCHAR*& Cmd, T& Value)
+std::enable_if_t<TModels<Meta::COnlineMetadataAvailable, T>::Value, bool> ParseOnlineExecParams(const TCHAR*& Cmd, T& Value, IOnlineServices* Services)
 {
 	bool bSuccess = true;
-	Meta::VisitFields(Value, [&Cmd, &bSuccess](const TCHAR* Name, auto& Field)
+
+	FString Token;
+	if (ParseTokenExt(Cmd, Token, true))
+	{
+		const TCHAR* TokenTCHAR =  *Token;
+		Meta::VisitFields(Value, [&TokenTCHAR, &bSuccess, Services](const TCHAR* Name, auto& Field)
 		{
-			bSuccess &= ParseOnlineExecParams(Cmd, Field);
+			bool bResult = ParseOnlineExecParams(TokenTCHAR, Field, Services);
+			if (!bResult)
+			{
+				UE_LOG(LogOnlineServices, Warning, TEXT("Failed to resolve inner-field %s"), Name);
+			}
+			bSuccess &= bResult;
 		});
+	}
+	else
+	{
+		bSuccess = false;
+	}
 
 	return bSuccess;
 }
@@ -394,8 +889,21 @@ public:
 
 	virtual bool Exec(UWorld* World, const TCHAR* Cmd, FOutputDevice& Ar) override
 	{
+		IOnlineServices* Services = &Interface->GetServices();
+
 		typename OpType::Params Params;
-		if (!Private::ParseOnlineExecParams(Cmd, Params))
+		bool bSuccess = true;
+		Meta::VisitFields(Params, [&Cmd, &bSuccess, &Ar, &Services](const TCHAR* Name, auto& Field)
+		{
+			bool bResult = Private::ParseOnlineExecParams(Cmd, Field, Services);
+			if (!bResult)
+			{
+				UE_LOG(LogOnlineServices, Warning, TEXT("Failed to resolve outer field %s"), Name);
+			}
+			bSuccess &= bResult;
+		});
+
+		if (!bSuccess)
 		{
 			Ar.Log(TEXT("Failed to parse params"));
 			return false;
diff --git a/Engine/Plugins/Online/Experimental/OnlineServices/Source/OnlineServicesInterface/Public/Online/OnlineUtils.h b/Engine/Plugins/Online/Experimental/OnlineServices/Source/OnlineServicesInterface/Public/Online/OnlineUtils.h
--- a/Engine/Plugins/Online/Experimental/OnlineServices/Source/OnlineServicesInterface/Public/Online/OnlineUtils.h
+++ b/Engine/Plugins/Online/Experimental/OnlineServices/Source/OnlineServicesInterface/Public/Online/OnlineUtils.h
@@ -2,6 +2,7 @@
 
 #pragma once
 
+#include "Online/OnlineAsyncOp.h"
 #include "Online/OnlineMeta.h"
 #include "Containers/UnrealString.h"
 
