OSSv2: Add Leaderboards interface and common implementation.

diff --git a/Engine/Plugins/Online/Experimental/OnlineServices/Source/OnlineServicesCommon/Private/Online/LeaderboardsCommon.cpp b/Engine/Plugins/Online/Experimental/OnlineServices/Source/OnlineServicesCommon/Private/Online/LeaderboardsCommon.cpp
new file mode 100644
--- /dev/null
+++ b/Engine/Plugins/Online/Experimental/OnlineServices/Source/OnlineServicesCommon/Private/Online/LeaderboardsCommon.cpp
@@ -0,0 +1,146 @@
+// Copyright Epic Games, Inc. All Rights Reserved.
+
+#include "Online/LeaderboardsCommon.h"
+
+#include "Online/Auth.h"
+#include "Online/OnlineServicesCommon.h"
+
+namespace UE::Online {
+
+const TCHAR* LexToString(ELeaderboardUpdateMethod Value)
+{
+	switch (Value)
+	{
+	case ELeaderboardUpdateMethod::Force:
+		return TEXT("Force");
+	default: checkNoEntry(); // Intentional fallthrough
+	case ELeaderboardUpdateMethod::KeepBest:
+		return TEXT("KeepBest");
+	}
+}
+
+void LexFromString(ELeaderboardUpdateMethod& OutValue, const TCHAR* InStr)
+{
+	if (FCString::Stricmp(InStr, TEXT("KeepBest")) == 0)
+	{
+		OutValue = ELeaderboardUpdateMethod::KeepBest;
+	}
+	else if (FCString::Stricmp(InStr, TEXT("Force")) == 0)
+	{
+		OutValue = ELeaderboardUpdateMethod::Force;
+	}
+	else
+	{
+		ensureMsgf(false, TEXT("Can't convert %s to ELeaderboardUpdateMethod"), InStr);
+		OutValue = ELeaderboardUpdateMethod::KeepBest;
+	}
+}
+
+const TCHAR* LexToString(ELeaderboardOrderMethod Value)
+{
+	switch (Value)
+	{
+	case ELeaderboardOrderMethod::Ascending:
+		return TEXT("Ascending");
+	default: checkNoEntry(); // Intentional fallthrough
+	case ELeaderboardOrderMethod::Descending:
+		return TEXT("Descending");
+	}
+}
+
+void LexFromString(ELeaderboardOrderMethod& OutValue, const TCHAR* InStr)
+{
+	if (FCString::Stricmp(InStr, TEXT("Ascending")) == 0)
+	{
+		OutValue = ELeaderboardOrderMethod::Ascending;
+	}
+	else if (FCString::Stricmp(InStr, TEXT("Descending")) == 0)
+	{
+		OutValue = ELeaderboardOrderMethod::Descending;
+	}
+	else
+	{
+		ensureMsgf(false, TEXT("Can't convert %s to ELeaderboardOrderMethod"), InStr);
+		OutValue = ELeaderboardOrderMethod::Descending;
+	}
+}
+
+FLeaderboardsCommon::FLeaderboardsCommon(FOnlineServicesCommon& InServices)
+	: TOnlineComponent(TEXT("Leaderboards"), InServices)
+{
+}
+
+void FLeaderboardsCommon::LoadConfig()
+{
+	Super::LoadConfig();
+
+	const TCHAR* ConfigSection = TEXT("OnlineServices.Leaderboards");
+
+	for (int StatIdx = 0;; StatIdx++)
+	{
+		FString LeaderboardName;
+		GConfig->GetString(ConfigSection, *FString::Printf(TEXT("LeaderboardDef_%d_Name"), StatIdx), LeaderboardName, GEngineIni);
+		if (LeaderboardName.IsEmpty())
+		{
+			break;
+		}
+
+		FLeaderboardDefinition& LeaderboardDefinition = LeaderboardDefinitions.Emplace(LeaderboardName);
+		LeaderboardDefinition.Name = MoveTemp(LeaderboardName);
+
+		GConfig->GetInt(ConfigSection, *FString::Printf(TEXT("LeaderboardDef_%d_Id"), StatIdx), LeaderboardDefinition.Id, GEngineIni);
+
+		FString LeaderboardUpdateMethod;
+		GConfig->GetString(ConfigSection, *FString::Printf(TEXT("LeaderboardDef_%d_UpdateMethod"), StatIdx), LeaderboardUpdateMethod, GEngineIni);
+		if (!LeaderboardUpdateMethod.IsEmpty())
+		{
+			LexFromString(LeaderboardDefinition.LeaderboardUpdateMethod, *LeaderboardUpdateMethod);
+		}
+
+		FString LeaderboardOrderMethod;
+		GConfig->GetString(ConfigSection, *FString::Printf(TEXT("LeaderboardDef_%d_OrderMethod"), StatIdx), LeaderboardOrderMethod, GEngineIni);
+		if (!LeaderboardOrderMethod.IsEmpty())
+		{
+			LexFromString(LeaderboardDefinition.LeaderboardOrderMethod, *LeaderboardOrderMethod);
+		}
+	}
+}
+
+void FLeaderboardsCommon::RegisterCommands()
+{
+	Super::RegisterCommands();
+
+	RegisterCommand(&FLeaderboardsCommon::ReadEntriesForUsers);
+	RegisterCommand(&FLeaderboardsCommon::ReadEntriesAroundRank);
+	RegisterCommand(&FLeaderboardsCommon::ReadEntriesAroundUser);
+}
+
+TOnlineAsyncOpHandle<FReadEntriesForUsers> FLeaderboardsCommon::ReadEntriesForUsers(FReadEntriesForUsers::Params&& Params)
+{
+	TOnlineAsyncOpRef<FReadEntriesForUsers> Operation = GetOp<FReadEntriesForUsers>(MoveTemp(Params));
+	Operation->SetError(Errors::NotImplemented());
+	return Operation->GetHandle();
+}
+
+TOnlineAsyncOpHandle<FReadEntriesAroundRank> FLeaderboardsCommon::ReadEntriesAroundRank(FReadEntriesAroundRank::Params&& Params)
+{
+	TOnlineAsyncOpRef<FReadEntriesAroundRank> Operation = GetOp<FReadEntriesAroundRank>(MoveTemp(Params));
+	Operation->SetError(Errors::NotImplemented());
+	return Operation->GetHandle();
+}
+
+TOnlineAsyncOpHandle<FReadEntriesAroundUser> FLeaderboardsCommon::ReadEntriesAroundUser(FReadEntriesAroundUser::Params&& Params)
+{
+	TOnlineAsyncOpRef<FReadEntriesAroundUser> Operation = GetOp<FReadEntriesAroundUser>(MoveTemp(Params));
+	Operation->SetError(Errors::NotImplemented());
+	return Operation->GetHandle();
+}
+
+TOnlineAsyncOpHandle<FWriteLeaderboardScores> FLeaderboardsCommon::WriteLeaderboardScores(FWriteLeaderboardScores::Params&& Params)
+{
+	TOnlineAsyncOpRef<FWriteLeaderboardScores> Operation = GetOp<FWriteLeaderboardScores>(MoveTemp(Params));
+	Operation->SetError(Errors::NotImplemented());
+	return Operation->GetHandle();
+}
+
+/* UE::Online */ }
diff --git a/Engine/Plugins/Online/Experimental/OnlineServices/Source/OnlineServicesCommon/Private/Online/OnlineServicesCommon.cpp b/Engine/Plugins/Online/Experimental/OnlineServices/Source/OnlineServicesCommon/Private/Online/OnlineServicesCommon.cpp
--- a/Engine/Plugins/Online/Experimental/OnlineServices/Source/OnlineServicesCommon/Private/Online/OnlineServicesCommon.cpp
+++ b/Engine/Plugins/Online/Experimental/OnlineServices/Source/OnlineServicesCommon/Private/Online/OnlineServicesCommon.cpp
@@ -70,6 +70,11 @@ IExternalUIPtr FOnlineServicesCommon::GetExternalUIInterface()
 	return IExternalUIPtr(AsShared(), Get<IExternalUI>());
 }
 
+ILeaderboardsPtr FOnlineServicesCommon::GetLeaderboardsInterface()
+{
+	return ILeaderboardsPtr(AsShared(), Get<ILeaderboards>());
+}
+
 ILobbiesPtr FOnlineServicesCommon::GetLobbiesInterface()
 {
 	return ILobbiesPtr(AsShared(), Get<ILobbies>());
diff --git a/Engine/Plugins/Online/Experimental/OnlineServices/Source/OnlineServicesCommon/Private/Online/StatsCommon.cpp b/Engine/Plugins/Online/Experimental/OnlineServices/Source/OnlineServicesCommon/Private/Online/StatsCommon.cpp
--- a/Engine/Plugins/Online/Experimental/OnlineServices/Source/OnlineServicesCommon/Private/Online/StatsCommon.cpp
+++ b/Engine/Plugins/Online/Experimental/OnlineServices/Source/OnlineServicesCommon/Private/Online/StatsCommon.cpp
@@ -48,64 +48,6 @@ void LexFromString(EStatModifyMethod& OutValue, const TCHAR* InStr)
 	}
 }
 
-const TCHAR* LexToString(EStatLeaderboardUpdateMethod Value)
-{
-	switch (Value)
-	{
-	case EStatLeaderboardUpdateMethod::Force:
-		return TEXT("Force");
-	default: checkNoEntry(); // Intentional fallthrough
-	case EStatLeaderboardUpdateMethod::KeepBest:
-		return TEXT("KeepBest");
-	}
-}
-
-void LexFromString(EStatLeaderboardUpdateMethod& OutValue, const TCHAR* InStr)
-{
-	if (FCString::Stricmp(InStr, TEXT("KeepBest")) == 0)
-	{
-		OutValue = EStatLeaderboardUpdateMethod::KeepBest;
-	}
-	else if (FCString::Stricmp(InStr, TEXT("Force")) == 0)
-	{
-		OutValue = EStatLeaderboardUpdateMethod::Force;
-	}
-	else
-	{
-		ensureMsgf(false, TEXT("Can't convert %s to EStatLeaderboardUpdateMethod"), InStr);
-		OutValue = EStatLeaderboardUpdateMethod::KeepBest;
-	}
-}
-
-const TCHAR* LexToString(EStatLeaderboardOrderMethod Value)
-{
-	switch (Value)
-	{
-	case EStatLeaderboardOrderMethod::Ascending:
-		return TEXT("Ascending");
-	default: checkNoEntry(); // Intentional fallthrough
-	case EStatLeaderboardOrderMethod::Descending:
-		return TEXT("Descending");
-	}
-}
-
-void LexFromString(EStatLeaderboardOrderMethod& OutValue, const TCHAR* InStr)
-{
-	if (FCString::Stricmp(InStr, TEXT("Ascending")) == 0)
-	{
-		OutValue = EStatLeaderboardOrderMethod::Ascending;
-	}
-	else if (FCString::Stricmp(InStr, TEXT("Descending")) == 0)
-	{
-		OutValue = EStatLeaderboardOrderMethod::Descending;
-	}
-	else
-	{
-		ensureMsgf(false, TEXT("Can't convert %s to EStatLeaderboardOrderMethod"), InStr);
-		OutValue = EStatLeaderboardOrderMethod::Descending;
-	}
-}
-
 const TCHAR* LexToString(EStatUsageFlags Value)
 {
 	switch (Value)
@@ -146,14 +88,7 @@ FStatsCommon::FStatsCommon(FOnlineServicesCommon& InServices)
 {
 }
 
-void FStatsCommon::Initialize()
-{
-	ReadStatDefinitionsFromConfig();
-
-	TOnlineComponent<IStats>::Initialize();
-}
-
-void FStatsCommon::ReadStatDefinitionsFromConfig()
+void FStatsCommon::LoadConfig()
 {
 	const TCHAR* ConfigSection = TEXT("OnlineServices.Stats");
 
@@ -189,23 +124,6 @@ void FStatsCommon::ReadStatDefinitionsFromConfig()
 		{
 			LexFromString(StatDefinition.ModifyMethod, *StatModifyMethod.ToString());
 		}
-
-		if (StatDefinition.UsageFlags & (uint32)EStatUsageFlags::Leaderboard)
-		{
-			FText StatLeaderboardUpdateMethod;
-			GConfig->GetText(ConfigSection, *FString::Printf(TEXT("StatDef_%d_LeaderboardUpdateMethod"), StatIdx), StatLeaderboardUpdateMethod, GEngineIni);
-			if (!StatLeaderboardUpdateMethod.IsEmpty())
-			{
-				LexFromString(StatDefinition.LeaderboardUpdateMethod, *StatLeaderboardUpdateMethod.ToString());
-			}
-
-			FText StatLeaderboardOrderMethod;
-			GConfig->GetText(ConfigSection, *FString::Printf(TEXT("StatDef_%d_LeaderboardOrderMethod"), StatIdx), StatLeaderboardOrderMethod, GEngineIni);
-			if (!StatLeaderboardOrderMethod.IsEmpty())
-			{
-				LexFromString(StatDefinition.LeaderboardOrderMethod, *StatLeaderboardOrderMethod.ToString());
-			}
-		}
 	}
 }
 
diff --git a/Engine/Plugins/Online/Experimental/OnlineServices/Source/OnlineServicesCommon/Public/Online/LeaderboardsCommon.h b/Engine/Plugins/Online/Experimental/OnlineServices/Source/OnlineServicesCommon/Public/Online/LeaderboardsCommon.h
new file mode 100644
--- /dev/null
+++ b/Engine/Plugins/Online/Experimental/OnlineServices/Source/OnlineServicesCommon/Public/Online/LeaderboardsCommon.h
@@ -0,0 +1,95 @@
+// Copyright Epic Games, Inc. All Rights Reserved.
+
+#pragma once
+
+#include "Online/Leaderboards.h"
+#include "Online/OnlineComponent.h"
+
+namespace UE::Online {
+
+enum class ELeaderboardUpdateMethod : uint8
+{
+	/** If current leaderboard score is better than the uploaded one, keep the current one */
+	KeepBest,
+	/** Leaderboard score is always replaced with uploaded value */
+	Force
+};
+
+const TCHAR* LexToString(ELeaderboardUpdateMethod Value);
+void LexFromString(ELeaderboardUpdateMethod& OutValue, const TCHAR* InStr);
+
+enum ELeaderboardOrderMethod : uint8
+{
+	Ascending,
+	Descending
+};
+
+const TCHAR* LexToString(ELeaderboardOrderMethod Value);
+void LexFromString(ELeaderboardOrderMethod& OutValue, const TCHAR* InStr);
+
+struct FLeaderboardDefinition
+{
+	/* The name of the leaderboard */
+	FString Name;
+	/* Corresponding leaderboard id on the platform if needed */
+	int32 Id = 0;
+
+	/* How the leaderboard score will be updated */
+	ELeaderboardUpdateMethod LeaderboardUpdateMethod = ELeaderboardUpdateMethod::Force;
+	/* How the leaderboard score will be ordered */
+	ELeaderboardOrderMethod LeaderboardOrderMethod = ELeaderboardOrderMethod::Descending;
+};
+
+struct FWriteLeaderboardScores
+{
+	static constexpr TCHAR Name[] = TEXT("WriteLeaderboardScores");
+
+	struct Params
+	{
+		FOnlineAccountIdHandle LocalUserId;
+		FString BoardName;
+		uint64 Score;
+	};
+
+	struct Result
+	{
+	};
+};
+
+class ONLINESERVICESCOMMON_API FLeaderboardsCommon : public TOnlineComponent<ILeaderboards>
+{
+public:
+	using Super = TOnlineComponent<ILeaderboards>;
+
+	FLeaderboardsCommon(FOnlineServicesCommon& InServices);
+
+	// TOnlineComponent
+	virtual void LoadConfig() override;
+	virtual void RegisterCommands() override;
+
+	// ILeaderboards
+	virtual TOnlineAsyncOpHandle<FReadEntriesForUsers> ReadEntriesForUsers(FReadEntriesForUsers::Params&& Params) override;
+	virtual TOnlineAsyncOpHandle<FReadEntriesAroundRank> ReadEntriesAroundRank(FReadEntriesAroundRank::Params&& Params) override;
+	virtual TOnlineAsyncOpHandle<FReadEntriesAroundUser> ReadEntriesAroundUser(FReadEntriesAroundUser::Params&& Params) override;
+
+	// Internal use only, for stats common implementation
+	virtual TOnlineAsyncOpHandle<FWriteLeaderboardScores> WriteLeaderboardScores(FWriteLeaderboardScores::Params&& Params);
+
+protected:
+	TMap<FString, FLeaderboardDefinition> LeaderboardDefinitions;
+};
+
+namespace Meta {
+
+BEGIN_ONLINE_STRUCT_META(FWriteLeaderboardScores::Params)
+	ONLINE_STRUCT_FIELD(FWriteLeaderboardScores::Params, LocalUserId),
+	ONLINE_STRUCT_FIELD(FWriteLeaderboardScores::Params, BoardName),
+	ONLINE_STRUCT_FIELD(FWriteLeaderboardScores::Params, Score)
+END_ONLINE_STRUCT_META()
+
+BEGIN_ONLINE_STRUCT_META(FWriteLeaderboardScores::Result)
+END_ONLINE_STRUCT_META()
+
+/* Meta*/ }
+
+/* UE::Online */ }
diff --git a/Engine/Plugins/Online/Experimental/OnlineServices/Source/OnlineServicesCommon/Public/Online/OnlineServicesCommon.h b/Engine/Plugins/Online/Experimental/OnlineServices/Source/OnlineServicesCommon/Public/Online/OnlineServicesCommon.h
--- a/Engine/Plugins/Online/Experimental/OnlineServices/Source/OnlineServicesCommon/Public/Online/OnlineServicesCommon.h
+++ b/Engine/Plugins/Online/Experimental/OnlineServices/Source/OnlineServicesCommon/Public/Online/OnlineServicesCommon.h
@@ -41,6 +41,7 @@ public:
 	virtual IFriendsPtr GetFriendsInterface() override;
 	virtual IPresencePtr GetPresenceInterface() override;
 	virtual IExternalUIPtr GetExternalUIInterface() override;
+	virtual ILeaderboardsPtr GetLeaderboardsInterface() override;
 	virtual ILobbiesPtr GetLobbiesInterface() override;
 	virtual ISessionsPtr GetSessionsInterface() override;
 	virtual IStatsPtr GetStatsInterface() override;
diff --git a/Engine/Plugins/Online/Experimental/OnlineServices/Source/OnlineServicesCommon/Public/Online/StatsCommon.h b/Engine/Plugins/Online/Experimental/OnlineServices/Source/OnlineServicesCommon/Public/Online/StatsCommon.h
--- a/Engine/Plugins/Online/Experimental/OnlineServices/Source/OnlineServicesCommon/Public/Online/StatsCommon.h
+++ b/Engine/Plugins/Online/Experimental/OnlineServices/Source/OnlineServicesCommon/Public/Online/StatsCommon.h
@@ -22,26 +22,6 @@ enum class EStatModifyMethod : uint8
 const TCHAR* LexToString(EStatModifyMethod Value);
 void LexFromString(EStatModifyMethod& OutValue, const TCHAR* InStr);
 
-enum class EStatLeaderboardUpdateMethod : uint8
-{
-	/** If current leaderboard score is better than the uploaded one, keep the current one */
-	KeepBest,
-	/** Leaderboard score is always replaced with uploaded value */
-	Force
-};
-
-const TCHAR* LexToString(EStatLeaderboardUpdateMethod Value);
-void LexFromString(EStatLeaderboardUpdateMethod& OutValue, const TCHAR* InStr);
-
-enum EStatLeaderboardOrderMethod : uint8
-{
-	Ascending,
-	Descending
-};
-
-const TCHAR* LexToString(EStatLeaderboardOrderMethod Value);
-void LexFromString(EStatLeaderboardOrderMethod& OutValue, const TCHAR* InStr);
-
 enum class EStatUsageFlags : uint8
 {
 	None = 0,
@@ -66,11 +46,6 @@ struct FStatDefinition
 
 	/* How the stat will be modified, only useful when EStatUsageFlags::Achievement is set in UsageFlags */
 	EStatModifyMethod ModifyMethod = EStatModifyMethod::Set;
-
-	/* How the leaderboard score will be updated, only useful when EStatUsageFlags::Leaderboard is set in UsageFlags */
-	EStatLeaderboardUpdateMethod LeaderboardUpdateMethod = EStatLeaderboardUpdateMethod::Force;
-	/* How the leaderboard score will be ordered, only useful when EStatUsageFlags::Leaderboard is set in UsageFlags */
-	EStatLeaderboardOrderMethod LeaderboardOrderMethod = EStatLeaderboardOrderMethod::Descending;
 };
 
 
@@ -82,7 +57,7 @@ public:
 	FStatsCommon(FOnlineServicesCommon& InServices);
 
 	// TOnlineComponent
-	virtual void Initialize() override;
+	virtual void LoadConfig() override;
 	virtual void RegisterCommands() override;
 
 	// IStats
@@ -94,8 +69,6 @@ public:
 #endif // !UE_BUILD_SHIPPING
 
 protected:
-	void ReadStatDefinitionsFromConfig();
-
 	TMap<FString, FStatDefinition> StatDefinitions;
 };
 
diff --git a/Engine/Plugins/Online/Experimental/OnlineServices/Source/OnlineServicesInterface/Public/Online/Leaderboards.h b/Engine/Plugins/Online/Experimental/OnlineServices/Source/OnlineServicesInterface/Public/Online/Leaderboards.h
new file mode 100644
--- /dev/null
+++ b/Engine/Plugins/Online/Experimental/OnlineServices/Source/OnlineServicesInterface/Public/Online/Leaderboards.h
@@ -0,0 +1,151 @@
+// Copyright Epic Games, Inc. All Rights Reserved.
+
+#pragma once
+
+#include "Online/CoreOnline.h"
+#include "Online/OnlineAsyncOpHandle.h"
+#include "Online/OnlineMeta.h"
+
+namespace UE::Online {
+
+struct FLeaderboardEntry
+{
+	/* The account id of this leaderboard entry */
+	FOnlineAccountIdHandle UserId;
+	/* The rank of this account */
+	int32 Rank;
+	/* The score of this account */
+	int64 Score;
+};
+
+struct FReadEntriesForUsers
+{
+	static constexpr TCHAR Name[] = TEXT("ReadEntriesForUsers");
+
+	struct Params
+	{
+		/* The account ids of users */
+		TArray<FOnlineAccountIdHandle> UserIds;
+		/* The leaderboard name */
+		FString BoardName;
+	};
+
+	struct Result
+	{
+		/* The result leaderboard entries */
+		TArray<FLeaderboardEntry> Entries;
+	};
+};
+
+struct FReadEntriesAroundRank
+{
+	static constexpr TCHAR Name[] = TEXT("ReadEntriesAroundRank");
+
+	struct Params
+	{
+		/* The rank index to start read */
+		int32 Rank;
+		/* How many entries to fetch */
+		int32 Limit;
+		/* The leaderboard to read */
+		FString BoardName;
+	};
+
+	struct Result
+	{
+		/* The result leaderboard entries */
+		TArray<FLeaderboardEntry> Entries;
+	};
+};
+
+struct FReadEntriesAroundUser
+{
+	static constexpr TCHAR Name[] = TEXT("ReadEntriesAroundUser");
+
+	struct Params
+	{
+		/* The account id of specified user */
+		FOnlineAccountIdHandle UserId;
+		/* The offset to the rank of user */
+		int32 Offset;
+		/* How many entries to fetch */
+		int32 Limit;
+		/* The leaderboard to read */
+		FString BoardName;
+	};
+
+	struct Result
+	{
+		/* The result leaderboard entries */
+		TArray<FLeaderboardEntry> Entries;
+	};
+};
+
+/**
+ * Interface definition for the leaderboards service 
+ */
+class ILeaderboards
+{
+public:
+	virtual ~ILeaderboards() {};
+
+	/**
+	 * Read leaderboard entries for specific users
+	 */
+	virtual TOnlineAsyncOpHandle<FReadEntriesForUsers> ReadEntriesForUsers(FReadEntriesForUsers::Params&& Params) = 0;
+
+	/**
+	 * Read leaderboard entries around specified rank index
+	 */
+	virtual TOnlineAsyncOpHandle<FReadEntriesAroundRank> ReadEntriesAroundRank(FReadEntriesAroundRank::Params&& Params) = 0;
+
+	/**
+	 * Read leaderboard entries around specified user. If the Offset is bigger than the Limit, the user 
+	 * may not appear in the returned entries. It makes sense in situations like getting prev page of users 
+	 * above my rank
+	 */
+	virtual TOnlineAsyncOpHandle<FReadEntriesAroundUser> ReadEntriesAroundUser(FReadEntriesAroundUser::Params&& Params) = 0;
+};
+
+
+namespace Meta {
+
+BEGIN_ONLINE_STRUCT_META(FLeaderboardEntry)
+	ONLINE_STRUCT_FIELD(FLeaderboardEntry, UserId),
+	ONLINE_STRUCT_FIELD(FLeaderboardEntry, Rank),
+	ONLINE_STRUCT_FIELD(FLeaderboardEntry, Score)
+END_ONLINE_STRUCT_META()
+
+BEGIN_ONLINE_STRUCT_META(FReadEntriesForUsers::Params)
+	ONLINE_STRUCT_FIELD(FReadEntriesForUsers::Params, UserIds),
+	ONLINE_STRUCT_FIELD(FReadEntriesForUsers::Params, BoardName)
+END_ONLINE_STRUCT_META()
+
+BEGIN_ONLINE_STRUCT_META(FReadEntriesForUsers::Result)
+	ONLINE_STRUCT_FIELD(FReadEntriesForUsers::Result, Entries)
+END_ONLINE_STRUCT_META()
+
+BEGIN_ONLINE_STRUCT_META(FReadEntriesAroundRank::Params)
+	ONLINE_STRUCT_FIELD(FReadEntriesAroundRank::Params, Rank),
+	ONLINE_STRUCT_FIELD(FReadEntriesAroundRank::Params, Limit),
+	ONLINE_STRUCT_FIELD(FReadEntriesAroundRank::Params, BoardName)
+END_ONLINE_STRUCT_META()
+
+BEGIN_ONLINE_STRUCT_META(FReadEntriesAroundRank::Result)
+	ONLINE_STRUCT_FIELD(FReadEntriesAroundRank::Result, Entries)
+END_ONLINE_STRUCT_META()
+
+BEGIN_ONLINE_STRUCT_META(FReadEntriesAroundUser::Params)
+	ONLINE_STRUCT_FIELD(FReadEntriesAroundUser::Params, UserId),
+	ONLINE_STRUCT_FIELD(FReadEntriesAroundUser::Params, Offset),
+	ONLINE_STRUCT_FIELD(FReadEntriesAroundUser::Params, Limit),
+	ONLINE_STRUCT_FIELD(FReadEntriesAroundUser::Params, BoardName)
+END_ONLINE_STRUCT_META()
+
+BEGIN_ONLINE_STRUCT_META(FReadEntriesAroundUser::Result)
+	ONLINE_STRUCT_FIELD(FReadEntriesAroundUser::Result, Entries)
+END_ONLINE_STRUCT_META()
+
+/* Meta*/ }
+
+/* UE::Online */ }
diff --git a/Engine/Plugins/Online/Experimental/OnlineServices/Source/OnlineServicesInterface/Public/Online/OnlineServices.h b/Engine/Plugins/Online/Experimental/OnlineServices/Source/OnlineServicesInterface/Public/Online/OnlineServices.h
--- a/Engine/Plugins/Online/Experimental/OnlineServices/Source/OnlineServicesInterface/Public/Online/OnlineServices.h
+++ b/Engine/Plugins/Online/Experimental/OnlineServices/Source/OnlineServicesInterface/Public/Online/OnlineServices.h
@@ -19,6 +19,7 @@ using IUserInfoPtr = TSharedPtr<class IUserInfo>;
 using IFriendsPtr = TSharedPtr<class IFriends>;
 using IPresencePtr = TSharedPtr<class IPresence>;
 using IExternalUIPtr = TSharedPtr<class IExternalUI>;
+using ILeaderboardsPtr = TSharedPtr<class ILeaderboards>;
 using ILobbiesPtr = TSharedPtr<class ILobbies>;
 using ISessionsPtr = TSharedPtr<class ISessions>;
 using IStatsPtr = TSharedPtr<class IStats>;
@@ -88,6 +89,12 @@ public:
 	 */
 	virtual IExternalUIPtr GetExternalUIInterface() = 0;
 
+	/**
+	 * Get the leaderboards implementation
+	 * @return leaderboards implementation, may be null if not implemented for this service
+	 */
+	virtual ILeaderboardsPtr GetLeaderboardsInterface() = 0;
+
 	/**
 	 * Get the lobbies implementation
 	 * @return lobbies implementation, may be null if not implemented for this service
