Change logging functions to use FString::JoinBy.

diff --git a/Engine/Plugins/Online/Experimental/OnlineServices/Source/OnlineServicesInterface/Public/Online/OnlineUtils.h b/Engine/Plugins/Online/Experimental/OnlineServices/Source/OnlineServicesInterface/Public/Online/OnlineUtils.h
--- a/Engine/Plugins/Online/Experimental/OnlineServices/Source/OnlineServicesInterface/Public/Online/OnlineUtils.h
+++ b/Engine/Plugins/Online/Experimental/OnlineServices/Source/OnlineServicesInterface/Public/Online/OnlineUtils.h
@@ -36,81 +36,25 @@ template <typename T> std::enable_if_t<TModels<Meta::COnlineMetadataAvailable, T
 template <typename T>
 FString ToLogString(const TArray<T>& Array)
 {
-	FString LogString;
-	LogString += TEXT("[ ");
-	bool bFirst = true;
-	for (const T& Value : Array)
-	{
-		if (bFirst)
-		{
-			bFirst = false;
-		}
-		else
-		{
-			LogString += TEXT(", ");
-		}
-		LogString += ToLogString(Value);
-	}
-	LogString += TEXT(" ]");
-	return LogString;
+	return FString::Printf(TEXT("[%s]"), *FString::JoinBy(Array, TEXT(", "), [](const T& Value) { return ToLogString(Value); }));
 }
 
 template <typename T>
 FString ToLogString(const TSet<T>& Set)
 {
-	FString LogString;
-	LogString += TEXT("[ ");
-	bool bFirst = true;
-	for (const T& Value : Set)
-	{
-		if (bFirst)
-		{
-			bFirst = false;
-		}
-		else
-		{
-			LogString += TEXT(", ");
-		}
-		LogString += ToLogString(Value);
-	}
-	LogString += TEXT(" ]");
-	return LogString;
+	return FString::Printf(TEXT("{%s}"), *FString::JoinBy(Set, TEXT(", "), [](const T& Value) { return ToLogString(Value); }));
 }
 
 template <typename K, typename V>
 FString ToLogString(const TMap<K, V>& Map)
 {
-	FString LogString;
-	LogString += TEXT("{ ");
-	bool bFirst = true;
-	for (const TPair<K, V>& Pair : Map)
-	{
-		if (bFirst)
-		{
-			bFirst = false;
-		}
-		else
-		{
-			LogString += TEXT(", ");
-		}
-		LogString += ToLogString(Pair.Key);
-		LogString += TEXT(": ");
-		LogString += ToLogString(Pair.Value);
-	}
-	LogString += TEXT(" }");
-	return LogString;
+	return FString::Printf(TEXT("{%s}"), *FString::JoinBy(Map, TEXT(", "), [](const TPair<K, V>& Pair) { return ToLogString(Pair); }));
 }
 
 template <typename T, typename U>
 FString ToLogString(const TPair<T, U>& Pair)
 {
-	FString LogString;
-	LogString += TEXT("[ ");
-	LogString += ToLogString(Pair.template Get<0>());
-	LogString += TEXT(", ");
-	LogString += ToLogString(Pair.template Get<1>());
-	LogString += TEXT(" ]");
-	return LogString;
+	return FString::Printf(TEXT("%s:%s"), *ToLogString(Pair.template Get<0>()), *ToLogString(Pair.template Get<1>()));
 }
 
 template <typename T, ESPMode Mode>
