OSSv2: In OnlineComponent, call LoadConfig() in Initialize function, instead of calling LoadConfig() when needed.

diff --git a/Engine/Plugins/Online/Experimental/OnlineServices/Source/OnlineServicesCommon/Public/Online/OnlineComponent.h b/Engine/Plugins/Online/Experimental/OnlineServices/Source/OnlineServicesCommon/Public/Online/OnlineComponent.h
--- a/Engine/Plugins/Online/Experimental/OnlineServices/Source/OnlineServicesCommon/Public/Online/OnlineComponent.h
+++ b/Engine/Plugins/Online/Experimental/OnlineServices/Source/OnlineServicesCommon/Public/Online/OnlineComponent.h
@@ -32,6 +32,7 @@ public:
 	{
 		Services.RegisterExecHandler(InterfaceName, MakeUnique<TOnlineComponentExecHandler<TOnlineComponent<ComponentType>>>(this));
 		RegisterCommands();
+		LoadConfig();
 	}
 	virtual void PostInitialize() override {}
 	virtual void LoadConfig() override {}
