Added missing module definition file for new module.

diff --git a/Engine/Plugins/Online/Experimental/OnlineServices/Source/OnlineServicesCommonEngineUtils/Private/Online/OnlineServicesCommonEngineUtilsModule.cpp b/Engine/Plugins/Online/Experimental/OnlineServices/Source/OnlineServicesCommonEngineUtils/Private/Online/OnlineServicesCommonEngineUtilsModule.cpp
new file mode 100644
--- /dev/null
+++ b/Engine/Plugins/Online/Experimental/OnlineServices/Source/OnlineServicesCommonEngineUtils/Private/Online/OnlineServicesCommonEngineUtilsModule.cpp
@@ -0,0 +1,20 @@
+// Copyright Epic Games, Inc. All Rights Reserved.
+
+#include "CoreMinimal.h"
+#include "Modules/ModuleManager.h"
+
+class FOnlineServicesCommonEngineUtilsModule : public IModuleInterface
+{
+public:
+	virtual void StartupModule()
+	{
+
+	}
+
+	virtual void ShutdownModule()
+	{
+
+	}
+};
+
+IMPLEMENT_MODULE(FOnlineServicesCommonEngineUtilsModule, OnlineServicesCommonEngineUtils);
\ No newline at end of file
