OSSv2: Use stats to unlock achievements for platforms which use title managed achievements.

diff --git a/Engine/Plugins/Online/Experimental/OnlineServices/Source/OnlineServicesCommon/Private/Online/AchievementsCommon.cpp b/Engine/Plugins/Online/Experimental/OnlineServices/Source/OnlineServicesCommon/Private/Online/AchievementsCommon.cpp
--- a/Engine/Plugins/Online/Experimental/OnlineServices/Source/OnlineServicesCommon/Private/Online/AchievementsCommon.cpp
+++ b/Engine/Plugins/Online/Experimental/OnlineServices/Source/OnlineServicesCommon/Private/Online/AchievementsCommon.cpp
@@ -5,14 +5,77 @@
 #include "Online/OnlineAsyncOp.h"
 #include "Online/OnlineErrorDefinitions.h"
 #include "Online/OnlineServicesCommon.h"
+#include "Online/StatsCommon.h"
 
 namespace UE::Online {
 
+bool FAchievementUnlockRule::ContainsStat(const FString& StatName) const
+{
+	return Conditions.ContainsByPredicate([&StatName](const FAchievementUnlockCondition& Condition) { return Condition.StatToCheck == StatName; });
+}
+
 FAchievementsCommon::FAchievementsCommon(FOnlineServicesCommon& InServices)
 	: TOnlineComponent(TEXT("Achievements"), InServices)
 {
 }
 
+void FAchievementsCommon::Initialize()
+{
+	TOnlineComponent<IAchievements>::Initialize();
+
+	StatEventHandle = Services.Get<FStatsCommon>()->OnStatsUpdated().Add([this](const FStatsUpdated& StatsUpdated) { UnlockAchievementsByStats(StatsUpdated); });
+}
+
+void FAchievementsCommon::Shutdown()
+{
+	StatEventHandle.Unbind();
+
+	TOnlineComponent<IAchievements>::Shutdown();
+}
+
+void FAchievementsCommon::LoadConfig()
+{
+	TOnlineComponent<IAchievements>::LoadConfig();
+
+	const TCHAR* ConfigSection = TEXT("OnlineServices.AchievementUnlockRules");
+	GConfig->GetBool(ConfigSection, TEXT("AchievementDef_IsTitleManaged"), bIsTitleManaged, GEngineIni);
+
+	if (bIsTitleManaged)
+	{
+		for (int RuleIdx = 0;; RuleIdx++)
+		{
+			FString AchievementToUnlock;
+			GConfig->GetString(ConfigSection, *FString::Printf(TEXT("AchievementDef_%d_UnlockRule_AchievementToUnlock"), RuleIdx), AchievementToUnlock, GEngineIni);
+
+			if (AchievementToUnlock.IsEmpty())
+			{
+				break;
+			}
+
+			FAchievementUnlockRule& AchievementUnlockRule = AchievementUnlockRules.Emplace_GetRef();
+			AchievementUnlockRule.AchievementToUnlock = MoveTemp(AchievementToUnlock);
+
+			FString ConditionsStr;
+			GConfig->GetString(ConfigSection, *FString::Printf(TEXT("AchievementDef_%d_UnlockRule_Conditions"), RuleIdx), ConditionsStr, GEngineIni);
+
+			bool bCullEmpty = true;
+			TArray<FString> ConditionStrs;
+			ConditionsStr.ParseIntoArray(ConditionStrs, TEXT(","), bCullEmpty);
+
+			for (const FString& ConditionStr : ConditionStrs)
+			{
+				FString StatToCheck;
+				FString UnlockThreshold;
+				ConditionStr.Split(TEXT(":"), &StatToCheck, &UnlockThreshold);
+
+				FAchievementUnlockCondition& Condition = AchievementUnlockRule.Conditions.Emplace_GetRef();
+				Condition.StatToCheck = MoveTemp(StatToCheck);
+				Condition.UnlockThreshold = (int64)FCString::Atoi(*UnlockThreshold);
+			}
+		}
+	}
+}
+
 void FAchievementsCommon::RegisterCommands()
 {
 	RegisterCommand(&FAchievementsCommon::QueryAchievementDefinitions);
@@ -48,7 +111,7 @@ TOnlineAsyncOpHandle<FQueryAchievementStates> FAchievementsCommon::QueryAchievem
 	return Operation->GetHandle();
 }
 
-TOnlineResult<FGetAchievementState> FAchievementsCommon::GetAchievementState(FGetAchievementState::Params&& Params)
+TOnlineResult<FGetAchievementState> FAchievementsCommon::GetAchievementState(FGetAchievementState::Params&& Params) const
 {
 	return TOnlineResult<FGetAchievementState>(Errors::NotImplemented());
 }
@@ -70,4 +133,139 @@ TOnlineEvent<void(const FAchievementStateUpdated&)> FAchievementsCommon::OnAchie
 	return OnAchievementStateUpdatedEvent;
 }
 
+void FAchievementsCommon::UnlockAchievementsByStats(const FStatsUpdated& StatsUpdated)
+{
+	if (!bIsTitleManaged)
+	{
+		return;
+	}
+
+	TArray<FString> StatNames;
+	TArray<FOnlineAccountIdHandle> UserIds;
+	for (const FUserStats& UserStats : StatsUpdated.UpdateUsersStats)
+	{
+		for (const TPair<FString, FStatValue>& StatPair : UserStats.Stats)
+		{
+			for (const FAchievementUnlockRule& AchievementUnlockRule : AchievementUnlockRules)
+			{
+				if (AchievementUnlockRule.ContainsStat(StatPair.Key))
+				{
+					for (const FAchievementUnlockCondition& Condition : AchievementUnlockRule.Conditions)
+					{
+						StatNames.AddUnique(Condition.StatToCheck);
+					}
+				}
+			}
+		}
+
+		UserIds.AddUnique(UserStats.UserId);
+	}
+
+	if (StatNames.IsEmpty() || UserIds.IsEmpty())
+	{
+		return;
+	}
+
+	FBatchQueryStats::Params BatchQueryStatsParam;
+	BatchQueryStatsParam.LocalUserId = StatsUpdated.LocalUserId;
+	BatchQueryStatsParam.TargetUserIds = MoveTemp(UserIds);
+	BatchQueryStatsParam.StatNames = MoveTemp(StatNames);
+
+	Services.Get<FStatsCommon>()->BatchQueryStats(MoveTemp(BatchQueryStatsParam))
+	.OnComplete([this](const TOnlineResult<FBatchQueryStats>& Result)
+	{
+		if (Result.IsOk())
+		{
+			const FBatchQueryStats::Result& BatchQueryStatsResult = Result.GetOkValue();
+			for (const FUserStats& UserStats : BatchQueryStatsResult.UsersStats)
+			{
+				TArray<FString> AchievementsToUnlock;
+				for (const TPair<FString, FStatValue>& StatPair : UserStats.Stats)
+				{
+					ExecuteUnlockRulesRelatedToStat(UserStats.UserId, StatPair.Key, UserStats.Stats, AchievementsToUnlock);
+				}
+
+				if (!AchievementsToUnlock.IsEmpty())
+				{
+					FUnlockAchievements::Params UnlockAchievementsParams;
+					UnlockAchievementsParams.LocalUserId = UserStats.UserId;
+					UnlockAchievementsParams.AchievementIds = MoveTemp(AchievementsToUnlock);
+					UnlockAchievements(MoveTemp(UnlockAchievementsParams));
+				}
+			}
+		}
+	});
+}
+
+void FAchievementsCommon::ExecuteUnlockRulesRelatedToStat(const FOnlineAccountIdHandle& UserId, const FString& StatName, const TMap<FString, FStatValue>& Stats, TArray<FString>& OutAchievementsToUnlock)
+{
+	for (const FAchievementUnlockRule& AchievementUnlockRule : AchievementUnlockRules)
+	{
+		if (AchievementUnlockRule.ContainsStat(StatName)
+			&& !IsUnlocked(UserId, AchievementUnlockRule.AchievementToUnlock) 
+			&& MeetUnlockCondition(AchievementUnlockRule, Stats))
+		{
+			OutAchievementsToUnlock.AddUnique(AchievementUnlockRule.AchievementToUnlock);
+		}
+	}
+}
+
+bool FAchievementsCommon::MeetUnlockCondition(FAchievementUnlockRule AchievementUnlockRule, const TMap<FString, FStatValue>& Stats)
+{
+	for (const FAchievementUnlockCondition& Condition : AchievementUnlockRule.Conditions)
+	{
+		if (const FStatDefinition* StatDefinition = Services.Get<FStatsCommon>()->GetStatDefinition(Condition.StatToCheck))
+		{
+			const FStatValue* StatValue = Stats.Find(Condition.StatToCheck);
+			if (!StatValue)
+			{
+				UE_LOG(LogTemp, Warning, TEXT("Can't find stat %s when check if it can unlock achievement."), *Condition.StatToCheck);
+				return false;
+			}
+
+			switch (StatDefinition->ModifyMethod)
+			{
+			case EStatModifyMethod::Sum: // Intentional fall through
+			case EStatModifyMethod::Largest:
+				if (StatValue->GetInt64() < Condition.UnlockThreshold.GetInt64())
+				{
+					return false;
+				}
+				break;
+			case EStatModifyMethod::Set:
+				if (StatValue->GetInt64() != Condition.UnlockThreshold.GetInt64())
+				{
+					return false;
+				}
+				break;
+			case EStatModifyMethod::Smallest:
+				if (StatValue->GetInt64() > Condition.UnlockThreshold.GetInt64())
+				{
+					return false;
+				}
+				break;
+			}
+		}
+	}
+
+	return true;
+}
+
+bool FAchievementsCommon::IsUnlocked(const FOnlineAccountIdHandle& UserId, const FString& AchievementName) const
+{
+	FGetAchievementState::Params Params;
+	Params.LocalUserId = UserId;
+	Params.AchievementId = AchievementName;
+	TOnlineResult<FGetAchievementState> Result = GetAchievementState(MoveTemp(Params));
+	if (Result.IsOk())
+	{
+		const FAchievementState& AchievementState = Result.GetOkValue().AchievementState;
+		return FMath::IsNearlyEqual(AchievementState.Progress, 1.0f);
+	}
+
+	UE_LOG(LogTemp, Warning, TEXT("Can't find state achievement %s when check if it's unlocked."), *AchievementName);
+
+	return false;
+}
+
 /* UE::Online */ }
diff --git a/Engine/Plugins/Online/Experimental/OnlineServices/Source/OnlineServicesCommon/Private/Online/StatsCommon.cpp b/Engine/Plugins/Online/Experimental/OnlineServices/Source/OnlineServicesCommon/Private/Online/StatsCommon.cpp
--- a/Engine/Plugins/Online/Experimental/OnlineServices/Source/OnlineServicesCommon/Private/Online/StatsCommon.cpp
+++ b/Engine/Plugins/Online/Experimental/OnlineServices/Source/OnlineServicesCommon/Private/Online/StatsCommon.cpp
@@ -169,4 +169,41 @@ TOnlineAsyncOpHandle<FResetStats> FStatsCommon::ResetStats(FResetStats::Params&&
 }
 #endif // !UE_BUILD_SHIPPING
 
+TOnlineResult<FGetCachedStats> FStatsCommon::GetCachedStats(FGetCachedStats::Params&& Params) const
+{
+	return TOnlineResult<FGetCachedStats>({ CachedUsersStats });
+}
+
+void FStatsCommon::CacheUserStats(const FUserStats& UserStats)
+{
+	if (FUserStats* ExistingUserStats = CachedUsersStats.FindByPredicate(FFindUserStatsByAccountId(UserStats.UserId)))
+	{
+		for (const TPair<FString, FStatValue>& StatPair : UserStats.Stats)
+		{
+			if (FStatValue* StatValue = ExistingUserStats->Stats.Find(StatPair.Key))
+			{
+				*StatValue = StatPair.Value;
+			}
+			else
+			{
+				ExistingUserStats->Stats.Add(StatPair);
+			}
+		}
+	}
+	else
+	{
+		CachedUsersStats.Emplace(UserStats);
+	}
+}
+
+TOnlineEvent<void(const FStatsUpdated&)> FStatsCommon::OnStatsUpdated() 
+{ 
+	return OnStatsUpdatedEvent; 
+}
+
+const FStatDefinition* FStatsCommon::GetStatDefinition(const FString& StatName) const 
+{ 
+	return StatDefinitions.Find(StatName); 
+}
+
 /* UE::Online */ }
diff --git a/Engine/Plugins/Online/Experimental/OnlineServices/Source/OnlineServicesCommon/Public/Online/AchievementsCommon.h b/Engine/Plugins/Online/Experimental/OnlineServices/Source/OnlineServicesCommon/Public/Online/AchievementsCommon.h
--- a/Engine/Plugins/Online/Experimental/OnlineServices/Source/OnlineServicesCommon/Public/Online/AchievementsCommon.h
+++ b/Engine/Plugins/Online/Experimental/OnlineServices/Source/OnlineServicesCommon/Public/Online/AchievementsCommon.h
@@ -4,11 +4,26 @@
 
 #include "Online/Achievements.h"
 #include "Online/OnlineComponent.h"
+#include "Online/Stats.h"
 
 namespace UE::Online {
 
 class FOnlineServicesCommon;
 
+struct FAchievementUnlockCondition
+{
+	FString StatToCheck;
+	FStatValue UnlockThreshold; // The unlock rule depends on Stat modification type
+};
+
+struct FAchievementUnlockRule
+{
+	FString AchievementToUnlock;
+	TArray<FAchievementUnlockCondition> Conditions;
+
+	bool ContainsStat(const FString& StatName) const;
+};
+
 class ONLINESERVICESCOMMON_API FAchievementsCommon : public TOnlineComponent<IAchievements>
 {
 public:
@@ -17,6 +32,9 @@ public:
 	FAchievementsCommon(FOnlineServicesCommon& InServices);
 
 	// TOnlineComponent
+	virtual void Initialize() override;
+	virtual void Shutdown() override;
+	virtual void LoadConfig() override;
 	virtual void RegisterCommands() override;
 
 	// IAchievements
@@ -24,13 +42,24 @@ public:
 	virtual TOnlineResult<FGetAchievementIds> GetAchievementIds(FGetAchievementIds::Params&& Params) override;
 	virtual TOnlineResult<FGetAchievementDefinition> GetAchievementDefinition(FGetAchievementDefinition::Params&& Params) override;
 	virtual TOnlineAsyncOpHandle<FQueryAchievementStates> QueryAchievementStates(FQueryAchievementStates::Params&& Params) override;
-	virtual TOnlineResult<FGetAchievementState> GetAchievementState(FGetAchievementState::Params&& Params) override;
+	virtual TOnlineResult<FGetAchievementState> GetAchievementState(FGetAchievementState::Params&& Params) const override;
 	virtual TOnlineAsyncOpHandle<FUnlockAchievements> UnlockAchievements(FUnlockAchievements::Params&& Params) override;
 	virtual TOnlineResult<FDisplayAchievementUI> DisplayAchievementUI(FDisplayAchievementUI::Params&& Params) override;
 	virtual TOnlineEvent<void(const FAchievementStateUpdated&)> OnAchievementStateUpdated() override;
 
 protected:
 	TOnlineEventCallable<void(const FAchievementStateUpdated&)> OnAchievementStateUpdatedEvent;
+
+	bool IsTitleManaged() const;
+	void UnlockAchievementsByStats(const FStatsUpdated& StatsUpdated);
+	void ExecuteUnlockRulesRelatedToStat(const FOnlineAccountIdHandle& UserId, const FString& StatName, const TMap<FString, FStatValue>& Stats, TArray<FString>& OutAchievementsToUnlock);
+	bool MeetUnlockCondition(FAchievementUnlockRule AchievementUnlockRule, const TMap<FString, FStatValue>& Stats);
+	bool IsUnlocked(const FOnlineAccountIdHandle& UserId, const FString& AchievementName) const;
+
+	TArray<FAchievementUnlockRule> AchievementUnlockRules;
+	FOnlineEventDelegateHandle StatEventHandle;
+
+	bool bIsTitleManaged = false;
 };
 
 /* UE::Online */ }
diff --git a/Engine/Plugins/Online/Experimental/OnlineServices/Source/OnlineServicesCommon/Public/Online/StatsCommon.h b/Engine/Plugins/Online/Experimental/OnlineServices/Source/OnlineServicesCommon/Public/Online/StatsCommon.h
--- a/Engine/Plugins/Online/Experimental/OnlineServices/Source/OnlineServicesCommon/Public/Online/StatsCommon.h
+++ b/Engine/Plugins/Online/Experimental/OnlineServices/Source/OnlineServicesCommon/Public/Online/StatsCommon.h
@@ -48,6 +48,20 @@ struct FStatDefinition
 	EStatModifyMethod ModifyMethod = EStatModifyMethod::Set;
 };
 
+struct FFindUserStatsByAccountId
+{
+	FFindUserStatsByAccountId(const FOnlineAccountIdHandle& InUserId)
+		: UserId(InUserId)
+	{
+	}
+
+	bool operator()(const FUserStats& UserStats) const
+	{
+		return UserStats.UserId == UserId;
+	}
+
+	FOnlineAccountIdHandle UserId;
+};
 
 class ONLINESERVICESCOMMON_API FStatsCommon : public TOnlineComponent<IStats>
 {
@@ -67,9 +81,18 @@ public:
 #if !UE_BUILD_SHIPPING
 	virtual TOnlineAsyncOpHandle<FResetStats> ResetStats(FResetStats::Params&& Params) override;
 #endif // !UE_BUILD_SHIPPING
+	virtual TOnlineResult<FGetCachedStats> GetCachedStats(FGetCachedStats::Params&& Params) const override;
+
+	virtual TOnlineEvent<void(const FStatsUpdated&)> OnStatsUpdated() override;
+	const FStatDefinition* GetStatDefinition(const FString& StatName) const;
 
 protected:
+	void CacheUserStats(const FUserStats& UserStats);
+
 	TMap<FString, FStatDefinition> StatDefinitions;
+	TOnlineEventCallable<void(const FStatsUpdated& StatsUpdated)> OnStatsUpdatedEvent;
+
+	TArray<FUserStats> CachedUsersStats;
 };
 
 /* UE::Online */ }
diff --git a/Engine/Plugins/Online/Experimental/OnlineServices/Source/OnlineServicesInterface/Public/Online/Achievements.h b/Engine/Plugins/Online/Experimental/OnlineServices/Source/OnlineServicesInterface/Public/Online/Achievements.h
--- a/Engine/Plugins/Online/Experimental/OnlineServices/Source/OnlineServicesInterface/Public/Online/Achievements.h
+++ b/Engine/Plugins/Online/Experimental/OnlineServices/Source/OnlineServicesInterface/Public/Online/Achievements.h
@@ -217,7 +217,7 @@ public:
 	 * Gets the state of an achievement by id for the given player.
 	 * Requires first calling QueryAchievementStates.
 	 */
-	virtual TOnlineResult<FGetAchievementState> GetAchievementState(FGetAchievementState::Params&& Params) = 0;
+	virtual TOnlineResult<FGetAchievementState> GetAchievementState(FGetAchievementState::Params&& Params) const = 0;
 
 	/**
 	 * Manually unlock achievements
diff --git a/Engine/Plugins/Online/Experimental/OnlineServices/Source/OnlineServicesInterface/Public/Online/Stats.h b/Engine/Plugins/Online/Experimental/OnlineServices/Source/OnlineServicesInterface/Public/Online/Stats.h
--- a/Engine/Plugins/Online/Experimental/OnlineServices/Source/OnlineServicesInterface/Public/Online/Stats.h
+++ b/Engine/Plugins/Online/Experimental/OnlineServices/Source/OnlineServicesInterface/Public/Online/Stats.h
@@ -95,6 +95,22 @@ struct FResetStats
 };
 #endif // !UE_BUILD_SHIPPING
 
+struct FGetCachedStats
+{
+	static constexpr TCHAR Name[] = TEXT("GetCachedStats");
+
+	struct Params
+	{
+	};
+
+	struct Result
+	{
+		TArray<FUserStats> UsersStats;
+	};
+};
+
+using FStatsUpdated = FUpdateStats::Params;
+
 // Stats interface is used to upload the stats to EOS or first party platforms, to complete corresponding InGame/Platform features such as Stats Query, Achievements, Leaderboard, etc.
 class IStats
 {
@@ -130,6 +146,18 @@ public:
 	 */
 	virtual TOnlineAsyncOpHandle<FResetStats> ResetStats(FResetStats::Params&& Params) = 0;
 #endif // !UE_BUILD_SHIPPING
+
+	/**
+	 * Event triggered when stats of user(s) changed
+	 */
+	virtual TOnlineEvent<void(const FStatsUpdated&)> OnStatsUpdated() = 0;
+
+	/**
+	 * Retrieve cached users' stats, which was retrieved when call QueryStats or BatchQueryStats, or when 
+	 * call UpdateStats for StatsNull impl. When choose to use StatsNull, make sure to listen to OnStatsUpdated
+	 * event and use this function to get all latest local stats and put them into SaveGame
+	 */
+	virtual TOnlineResult<FGetCachedStats> GetCachedStats(FGetCachedStats::Params&& Params) const = 0;
 };
 
 namespace Meta {
@@ -172,6 +200,13 @@ END_ONLINE_STRUCT_META()
 
 BEGIN_ONLINE_STRUCT_META(FResetStats::Result)
 END_ONLINE_STRUCT_META()
+
+BEGIN_ONLINE_STRUCT_META(FGetCachedStats::Params)
+END_ONLINE_STRUCT_META()
+
+BEGIN_ONLINE_STRUCT_META(FGetCachedStats::Result)
+	ONLINE_STRUCT_FIELD(FGetCachedStats::Result, UsersStats)
+END_ONLINE_STRUCT_META()
 #endif // !UE_BUILD_SHIPPING
 
 /* Meta*/ }
diff --git a/Engine/Plugins/Online/Experimental/OnlineServicesEOSGS/Source/Private/Online/AchievementsEOSGS.cpp b/Engine/Plugins/Online/Experimental/OnlineServicesEOSGS/Source/Private/Online/AchievementsEOSGS.cpp
--- a/Engine/Plugins/Online/Experimental/OnlineServicesEOSGS/Source/Private/Online/AchievementsEOSGS.cpp
+++ b/Engine/Plugins/Online/Experimental/OnlineServicesEOSGS/Source/Private/Online/AchievementsEOSGS.cpp
@@ -286,7 +286,7 @@ TOnlineAsyncOpHandle<FQueryAchievementStates> FAchievementsEOSGS::QueryAchieveme
 	return Op->GetHandle();
 }
 
-TOnlineResult<FGetAchievementState> FAchievementsEOSGS::GetAchievementState(FGetAchievementState::Params&& Params)
+TOnlineResult<FGetAchievementState> FAchievementsEOSGS::GetAchievementState(FGetAchievementState::Params&& Params) const
 {
 	if (!Services.Get<FAuthEOSGS>()->IsLoggedIn(Params.LocalUserId))
 	{
diff --git a/Engine/Plugins/Online/Experimental/OnlineServicesEOSGS/Source/Private/Online/StatsEOSGS.cpp b/Engine/Plugins/Online/Experimental/OnlineServicesEOSGS/Source/Private/Online/StatsEOSGS.cpp
--- a/Engine/Plugins/Online/Experimental/OnlineServicesEOSGS/Source/Private/Online/StatsEOSGS.cpp
+++ b/Engine/Plugins/Online/Experimental/OnlineServicesEOSGS/Source/Private/Online/StatsEOSGS.cpp
@@ -91,6 +91,7 @@ TOnlineAsyncOpHandle<FUpdateStats> FStatsEOSGS::UpdateStats(FUpdateStats::Params
 	Op->Then([this](TOnlineAsyncOp<FUpdateStats>& InAsyncOp)
 	{
 		InAsyncOp.SetResult({});
+		OnStatsUpdatedEvent.Broadcast(InAsyncOp.GetParams());
 	});
 
 	Op->Enqueue(GetSerialQueue());
@@ -181,6 +182,12 @@ TOnlineAsyncOpHandle<FQueryStats> FStatsEOSGS::QueryStats(FQueryStats::Params&&
 
 		FQueryStats::Result Result;
 		Private::ReadStatsFromEOSResult(StatsHandle, Data, InAsyncOp.GetParams().StatNames, Result.Stats);
+
+		FUserStats UserStats;
+		UserStats.UserId = InAsyncOp.GetParams().TargetUserId;
+		UserStats.Stats = Result.Stats;
+		CacheUserStats(UserStats);
+
 		InAsyncOp.SetResult(MoveTemp(Result));
 	});
 
@@ -216,21 +223,22 @@ TOnlineAsyncOpHandle<FBatchQueryStats> FStatsEOSGS::BatchQueryStats(FBatchQueryS
 			if (Data->ResultCode != EOS_EResult::EOS_Success)
 			{
 				UE_LOG(LogTemp, Warning, TEXT("EOS_Stats_QueryStats failed with result=[%s]"), *LexToString(Data->ResultCode));
-				CachedUsersStats.Reset();
+				BatchQueriedUsersStats.Reset();
 				InAsyncOp.SetError(FromEOSError(Data->ResultCode));
 				return;
 			}
 
-			FUserStats& UserStats = CachedUsersStats.Emplace_GetRef();
+			FUserStats& UserStats = BatchQueriedUsersStats.Emplace_GetRef();
 
 			Private::ReadStatsFromEOSResult(StatsHandle, Data, InAsyncOp.GetParams().StatNames, UserStats.Stats);
+			CacheUserStats(UserStats);
 		});
 	}
 
 	Op->Then([this](TOnlineAsyncOp<FBatchQueryStats>& InAsyncOp)
 	{
 		FBatchQueryStats::Result Result;
-		Result.UsersStats = MoveTemp(CachedUsersStats);
+		Result.UsersStats = MoveTemp(BatchQueriedUsersStats);
 		InAsyncOp.SetResult(MoveTemp(Result));
 	});
 
diff --git a/Engine/Plugins/Online/Experimental/OnlineServicesEOSGS/Source/Public/Online/AchievementsEOSGS.h b/Engine/Plugins/Online/Experimental/OnlineServicesEOSGS/Source/Public/Online/AchievementsEOSGS.h
--- a/Engine/Plugins/Online/Experimental/OnlineServicesEOSGS/Source/Public/Online/AchievementsEOSGS.h
+++ b/Engine/Plugins/Online/Experimental/OnlineServicesEOSGS/Source/Public/Online/AchievementsEOSGS.h
@@ -31,7 +31,7 @@ public:
 	virtual TOnlineResult<FGetAchievementIds> GetAchievementIds(FGetAchievementIds::Params&& Params) override;
 	virtual TOnlineResult<FGetAchievementDefinition> GetAchievementDefinition(FGetAchievementDefinition::Params&& Params) override;
 	virtual TOnlineAsyncOpHandle<FQueryAchievementStates> QueryAchievementStates(FQueryAchievementStates::Params&& Params) override;
-	virtual TOnlineResult<FGetAchievementState> GetAchievementState(FGetAchievementState::Params&& Params) override;
+	virtual TOnlineResult<FGetAchievementState> GetAchievementState(FGetAchievementState::Params&& Params) const override;
 	virtual TOnlineAsyncOpHandle<FUnlockAchievements> UnlockAchievements(FUnlockAchievements::Params&& Params) override;
 
 protected:
diff --git a/Engine/Plugins/Online/Experimental/OnlineServicesEOSGS/Source/Public/Online/StatsEOSGS.h b/Engine/Plugins/Online/Experimental/OnlineServicesEOSGS/Source/Public/Online/StatsEOSGS.h
--- a/Engine/Plugins/Online/Experimental/OnlineServicesEOSGS/Source/Public/Online/StatsEOSGS.h
+++ b/Engine/Plugins/Online/Experimental/OnlineServicesEOSGS/Source/Public/Online/StatsEOSGS.h
@@ -31,7 +31,7 @@ public:
 
 private:
 	EOS_HStats StatsHandle = nullptr;
-	TArray<FUserStats> CachedUsersStats;
+	TArray<FUserStats> BatchQueriedUsersStats;
 };
 
 /* UE::Online */ }
diff --git a/Engine/Plugins/Online/Experimental/OnlineServicesNull/Source/Private/Online/AchievementsNull.cpp b/Engine/Plugins/Online/Experimental/OnlineServicesNull/Source/Private/Online/AchievementsNull.cpp
--- a/Engine/Plugins/Online/Experimental/OnlineServicesNull/Source/Private/Online/AchievementsNull.cpp
+++ b/Engine/Plugins/Online/Experimental/OnlineServicesNull/Source/Private/Online/AchievementsNull.cpp
@@ -160,7 +160,7 @@ TOnlineAsyncOpHandle<FQueryAchievementStates> FAchievementsNull::QueryAchievemen
 	return Op->GetHandle();
 }
 
-TOnlineResult<FGetAchievementState> FAchievementsNull::GetAchievementState(FGetAchievementState::Params&& Params)
+TOnlineResult<FGetAchievementState> FAchievementsNull::GetAchievementState(FGetAchievementState::Params&& Params) const
 {
 	if (!Services.Get<FAuthNull>()->IsLoggedIn(Params.LocalUserId))
 	{
diff --git a/Engine/Plugins/Online/Experimental/OnlineServicesNull/Source/Private/Online/StatsNull.cpp b/Engine/Plugins/Online/Experimental/OnlineServicesNull/Source/Private/Online/StatsNull.cpp
--- a/Engine/Plugins/Online/Experimental/OnlineServicesNull/Source/Private/Online/StatsNull.cpp
+++ b/Engine/Plugins/Online/Experimental/OnlineServicesNull/Source/Private/Online/StatsNull.cpp
@@ -12,21 +12,6 @@ FStatsNull::FStatsNull(FOnlineServicesNull& InOwningSubsystem)
 {
 }
 
-struct FFindUserStatsNull
-{
-	FFindUserStatsNull(const FOnlineAccountIdHandle& InUserId)
-		: UserId(InUserId)
-	{
-	}
-
-	bool operator()(const FUserStats& UserStats)
-	{
-		return UserStats.UserId == UserId;
-	}
-
-	FOnlineAccountIdHandle UserId;
-};
-
 TOnlineAsyncOpHandle<FUpdateStats> FStatsNull::UpdateStats(FUpdateStats::Params&& Params)
 {
 	TOnlineAsyncOpRef<FUpdateStats> Op = GetOp<FUpdateStats>(MoveTemp(Params));
@@ -41,7 +26,7 @@ TOnlineAsyncOpHandle<FUpdateStats> FStatsNull::UpdateStats(FUpdateStats::Params&
 	{
 		for (const FUserStats& UpdateUserStats : InAsyncOp.GetParams().UpdateUsersStats)
 		{
-			FUserStats* ExistingUserStats = UsersStats.FindByPredicate(FFindUserStatsNull(UpdateUserStats.UserId));
+			FUserStats* ExistingUserStats = UsersStats.FindByPredicate(FFindUserStatsByAccountId(UpdateUserStats.UserId));
 			if (!ExistingUserStats)
 			{
 				ExistingUserStats = &UsersStats.Emplace_GetRef();
@@ -103,9 +88,12 @@ TOnlineAsyncOpHandle<FUpdateStats> FStatsNull::UpdateStats(FUpdateStats::Params&
 					}
 				}
 			}
+
+			CacheUserStats(*ExistingUserStats);
 		}
 
 		InAsyncOp.SetResult({});
+		OnStatsUpdatedEvent.Broadcast(InAsyncOp.GetParams());
 	})
 	.Enqueue(GetSerialQueue());
 
@@ -126,7 +114,7 @@ TOnlineAsyncOpHandle<FQueryStats> FStatsNull::QueryStats(FQueryStats::Params&& P
 	{
 		FQueryStats::Result Result;
 
-		if (FUserStats* ExistingUserStats = UsersStats.FindByPredicate(FFindUserStatsNull(InAsyncOp.GetParams().TargetUserId)))
+		if (FUserStats* ExistingUserStats = UsersStats.FindByPredicate(FFindUserStatsByAccountId(InAsyncOp.GetParams().TargetUserId)))
 		{
 			Result.Stats = ExistingUserStats->Stats;
 		}
@@ -153,7 +141,7 @@ TOnlineAsyncOpHandle<FBatchQueryStats> FStatsNull::BatchQueryStats(FBatchQuerySt
 
 		for (const FOnlineAccountIdHandle& TargetUserId : InAsyncOp.GetParams().TargetUserIds)
 		{
-			if (FUserStats* ExistingUserStats = UsersStats.FindByPredicate(FFindUserStatsNull(TargetUserId)))
+			if (FUserStats* ExistingUserStats = UsersStats.FindByPredicate(FFindUserStatsByAccountId(TargetUserId)))
 			{
 				Result.UsersStats.Emplace(*ExistingUserStats);
 			}
@@ -179,11 +167,17 @@ TOnlineAsyncOpHandle<FResetStats> FStatsNull::ResetStats(FResetStats::Params&& P
 
 	Op->Then([this](TOnlineAsyncOp<FResetStats>& InAsyncOp)
 	{
-		uint32 Index = UsersStats.IndexOfByPredicate(FFindUserStatsNull(InAsyncOp.GetParams().LocalUserId));
+		uint32 Index = UsersStats.IndexOfByPredicate(FFindUserStatsByAccountId(InAsyncOp.GetParams().LocalUserId));
 		if (Index != INDEX_NONE)
 		{
 			UsersStats.RemoveAt(Index);
 		}
+
+		FFindUserStatsByAccountId FindUserStatsByAccountId(InAsyncOp.GetParams().LocalUserId);
+
+		UsersStats.RemoveAll(FindUserStatsByAccountId);
+		CachedUsersStats.RemoveAll(FindUserStatsByAccountId);
+
 		InAsyncOp.SetResult({});
 	})
 	.Enqueue(GetSerialQueue());
diff --git a/Engine/Plugins/Online/Experimental/OnlineServicesNull/Source/Public/Online/AchievementsNull.h b/Engine/Plugins/Online/Experimental/OnlineServicesNull/Source/Public/Online/AchievementsNull.h
--- a/Engine/Plugins/Online/Experimental/OnlineServicesNull/Source/Public/Online/AchievementsNull.h
+++ b/Engine/Plugins/Online/Experimental/OnlineServicesNull/Source/Public/Online/AchievementsNull.h
@@ -21,7 +21,7 @@ public:
 	virtual TOnlineResult<FGetAchievementIds> GetAchievementIds(FGetAchievementIds::Params&& Params) override;
 	virtual TOnlineResult<FGetAchievementDefinition> GetAchievementDefinition(FGetAchievementDefinition::Params&& Params) override;
 	virtual TOnlineAsyncOpHandle<FQueryAchievementStates> QueryAchievementStates(FQueryAchievementStates::Params&& Params) override;
-	virtual TOnlineResult<FGetAchievementState> GetAchievementState(FGetAchievementState::Params&& Params) override;
+	virtual TOnlineResult<FGetAchievementState> GetAchievementState(FGetAchievementState::Params&& Params) const override;
 	virtual TOnlineAsyncOpHandle<FUnlockAchievements> UnlockAchievements(FUnlockAchievements::Params&& Params) override;
 	virtual TOnlineResult<FDisplayAchievementUI> DisplayAchievementUI(FDisplayAchievementUI::Params&& Params) override;
 
diff --git a/Engine/Plugins/Online/Experimental/OnlineServicesNull/Source/Public/Online/StatsNull.h b/Engine/Plugins/Online/Experimental/OnlineServicesNull/Source/Public/Online/StatsNull.h
--- a/Engine/Plugins/Online/Experimental/OnlineServicesNull/Source/Public/Online/StatsNull.h
+++ b/Engine/Plugins/Online/Experimental/OnlineServicesNull/Source/Public/Online/StatsNull.h
@@ -25,6 +25,7 @@ public:
 #endif // !UE_BUILD_SHIPPING
 
 protected:
+	// TODO: Save UsersStats into local user profile
 	TArray<FUserStats> UsersStats;
 };
 
diff --git a/Engine/Plugins/Online/Experimental/OnlineServicesOSSAdapter/Source/Private/Online/AchievementsOSSAdapter.cpp b/Engine/Plugins/Online/Experimental/OnlineServicesOSSAdapter/Source/Private/Online/AchievementsOSSAdapter.cpp
--- a/Engine/Plugins/Online/Experimental/OnlineServicesOSSAdapter/Source/Private/Online/AchievementsOSSAdapter.cpp
+++ b/Engine/Plugins/Online/Experimental/OnlineServicesOSSAdapter/Source/Private/Online/AchievementsOSSAdapter.cpp
@@ -14,6 +14,11 @@
 
 namespace UE::Online {
 
+void FAchievementsOSSAdapter::Initialize()
+{
+	Super::Initialize();
+}
+
 void FAchievementsOSSAdapter::PostInitialize()
 {
 	Super::PostInitialize();
@@ -273,7 +278,7 @@ TOnlineAsyncOpHandle<FQueryAchievementStates> FAchievementsOSSAdapter::QueryAchi
 	return Op->GetHandle();	
 }
 
-TOnlineResult<FGetAchievementState> FAchievementsOSSAdapter::GetAchievementState(FGetAchievementState::Params&& Params)
+TOnlineResult<FGetAchievementState> FAchievementsOSSAdapter::GetAchievementState(FGetAchievementState::Params&& Params) const
 {
 	const FUniqueNetIdPtr LocalUserId = Services.Get<FAuthOSSAdapter>()->GetUniqueNetId(Params.LocalUserId);
 	if (!LocalUserId)
diff --git a/Engine/Plugins/Online/Experimental/OnlineServicesOSSAdapter/Source/Private/Online/StatsOSSAdapter.cpp b/Engine/Plugins/Online/Experimental/OnlineServicesOSSAdapter/Source/Private/Online/StatsOSSAdapter.cpp
--- a/Engine/Plugins/Online/Experimental/OnlineServicesOSSAdapter/Source/Private/Online/StatsOSSAdapter.cpp
+++ b/Engine/Plugins/Online/Experimental/OnlineServicesOSSAdapter/Source/Private/Online/StatsOSSAdapter.cpp
@@ -126,7 +126,7 @@ TOnlineAsyncOpHandle<FUpdateStats> FStatsOSSAdapter::UpdateStats(FUpdateStats::P
 			return;
 		}
 
-		StatsInterface->UpdateStats(LocalUserId.ToSharedRef(), UpdateUserStatsV1, *MakeDelegateAdapter(this, [WeakOp = Op.AsWeak()](const FOnlineErrorOss& OnlineError) mutable
+		StatsInterface->UpdateStats(LocalUserId.ToSharedRef(), UpdateUserStatsV1, *MakeDelegateAdapter(this, [this, WeakOp = Op.AsWeak()](const FOnlineErrorOss& OnlineError) mutable
 		{
 			if (TSharedPtr<TOnlineAsyncOp<FUpdateStats>> PinnedOp = WeakOp.Pin())
 			{
@@ -137,6 +137,8 @@ TOnlineAsyncOpHandle<FUpdateStats> FStatsOSSAdapter::UpdateStats(FUpdateStats::P
 				}
 
 				PinnedOp->SetResult({});
+
+				OnStatsUpdatedEvent.Broadcast(PinnedOp->GetParams());
 			}
 		}));
 	}) 
@@ -178,6 +180,11 @@ TOnlineAsyncOpHandle<FQueryStats> FStatsOSSAdapter::QueryStats(FQueryStats::Para
 					Result.Stats.Emplace(StatPair.Key, MoveTemp(StatValue));
 				}
 
+				FUserStats UserStats;
+				UserStats.UserId = PinnedOp->GetParams().TargetUserId;
+				UserStats.Stats = Result.Stats;
+				CacheUserStats(UserStats);
+
 				PinnedOp->SetResult(MoveTemp(Result));
 			}
 		}));
@@ -235,6 +242,8 @@ TOnlineAsyncOpHandle<FBatchQueryStats> FStatsOSSAdapter::BatchQueryStats(FBatchQ
 						ConvertStatValueV1ToStatValueV2(StatPair.Value, StatValue);
 						UserStats.Stats.Emplace(StatPair.Key, MoveTemp(StatValue));
 					}
+
+					CacheUserStats(UserStats);
 				}
 
 				PinnedOp->SetResult(MoveTemp(Result));
diff --git a/Engine/Plugins/Online/Experimental/OnlineServicesOSSAdapter/Source/Public/Online/AchievementsOSSAdapter.h b/Engine/Plugins/Online/Experimental/OnlineServicesOSSAdapter/Source/Public/Online/AchievementsOSSAdapter.h
--- a/Engine/Plugins/Online/Experimental/OnlineServicesOSSAdapter/Source/Public/Online/AchievementsOSSAdapter.h
+++ b/Engine/Plugins/Online/Experimental/OnlineServicesOSSAdapter/Source/Public/Online/AchievementsOSSAdapter.h
@@ -19,6 +19,7 @@ public:
 	using FAchievementsCommon::FAchievementsCommon;
 
 	// IOnlineComponent
+	virtual void Initialize() override;
 	virtual void PostInitialize() override;
 	virtual void PreShutdown() override;
 
@@ -27,7 +28,7 @@ public:
 	virtual TOnlineResult<FGetAchievementIds> GetAchievementIds(FGetAchievementIds::Params&& Params) override;
 	virtual TOnlineResult<FGetAchievementDefinition> GetAchievementDefinition(FGetAchievementDefinition::Params&& Params) override;
 	virtual TOnlineAsyncOpHandle<FQueryAchievementStates> QueryAchievementStates(FQueryAchievementStates::Params&& Params) override;
-	virtual TOnlineResult<FGetAchievementState> GetAchievementState(FGetAchievementState::Params&& Params) override;
+	virtual TOnlineResult<FGetAchievementState> GetAchievementState(FGetAchievementState::Params&& Params) const override;
 	// Intentionally unimplemented, all v1 achis are stats based
 	//virtual TOnlineAsyncOpHandle<FUnlockAchievements> UnlockAchievements(FUnlockAchievements::Params&& Params) override;
 	virtual TOnlineResult<FDisplayAchievementUI> DisplayAchievementUI(FDisplayAchievementUI::Params&& Params) override;
