- Add Ptr and WeakPtr aliases for FTitleFileContents

diff --git a/Engine/Plugins/Online/Experimental/OnlineServices/Source/OnlineServicesInterface/Public/Online/TitleFile.h b/Engine/Plugins/Online/Experimental/OnlineServices/Source/OnlineServicesInterface/Public/Online/TitleFile.h
--- a/Engine/Plugins/Online/Experimental/OnlineServices/Source/OnlineServicesInterface/Public/Online/TitleFile.h
+++ b/Engine/Plugins/Online/Experimental/OnlineServices/Source/OnlineServicesInterface/Public/Online/TitleFile.h
@@ -42,6 +42,8 @@ struct FTitleFileGetEnumeratedFiles
 
 using FTitleFileContents = TArray<uint8>;
 using FTitleFileContentsRef = TSharedRef<const FTitleFileContents>;
+using FTitleFileContentsPtr = TSharedPtr<const FTitleFileContents>;
+using FTitleFileContentsWeakPtr = TWeakPtr<const FTitleFileContents>;
 
 struct FTitleFileReadFile
 {
