- Add Ptr and WeakPtr typedefs for FUserFileContents

diff --git a/Engine/Plugins/Online/Experimental/OnlineServices/Source/OnlineServicesInterface/Public/Online/UserFile.h b/Engine/Plugins/Online/Experimental/OnlineServices/Source/OnlineServicesInterface/Public/Online/UserFile.h
--- a/Engine/Plugins/Online/Experimental/OnlineServices/Source/OnlineServicesInterface/Public/Online/UserFile.h
+++ b/Engine/Plugins/Online/Experimental/OnlineServices/Source/OnlineServicesInterface/Public/Online/UserFile.h
@@ -42,6 +42,8 @@ struct FUserFileGetEnumeratedFiles
 
 using FUserFileContents = TArray<uint8>;
 using FUserFileContentsRef = TSharedRef<const FUserFileContents>;
+using FUserFileContentsPtr = TSharedPtr<const FUserFileContents>;
+using FUserFileContentsWeakPtr = TWeakPtr<const FUserFileContents>;
 
 struct FUserFileReadFile
 {
