Header files and common implementation for the Commerce interface

diff --git a/Engine/Plugins/Online/Experimental/OnlineServices/Source/OnlineServicesCommon/Private/Online/CommerceCommon.cpp b/Engine/Plugins/Online/Experimental/OnlineServices/Source/OnlineServicesCommon/Private/Online/CommerceCommon.cpp
new file mode 100644
--- /dev/null
+++ b/Engine/Plugins/Online/Experimental/OnlineServices/Source/OnlineServicesCommon/Private/Online/CommerceCommon.cpp
@@ -0,0 +1,106 @@
+// Copyright Epic Games, Inc. All Rights Reserved.
+
+#include "Online/CommerceCommon.h"
+#include "Online/OnlineAsyncOp.h"
+#include "Online/OnlineErrorDefinitions.h"
+#include "Online/OnlineServicesCommon.h"
+
+namespace UE::Online {
+
+FCommerceCommon::FCommerceCommon(FOnlineServicesCommon& InServices)
+: TOnlineComponent(TEXT("Commerce"), InServices)
+{
+}
+
+void FCommerceCommon::RegisterCommands()
+{
+	RegisterCommand(&FCommerceCommon::QueryOffers);
+	RegisterCommand(&FCommerceCommon::QueryOffersById);
+	RegisterCommand(&FCommerceCommon::GetOffers);
+	RegisterCommand(&FCommerceCommon::GetOffersById);
+	RegisterCommand(&FCommerceCommon::ShowStoreUI);
+	RegisterCommand(&FCommerceCommon::Checkout);
+	RegisterCommand(&FCommerceCommon::QueryTransactionEntitlements);
+	RegisterCommand(&FCommerceCommon::QueryEntitlements);
+	RegisterCommand(&FCommerceCommon::GetEntitlements);
+	RegisterCommand(&FCommerceCommon::RedeemEntitlement);
+	RegisterCommand(&FCommerceCommon::RetrieveS2SToken);
+}
+
+TOnlineAsyncOpHandle<FCommerceQueryOffers> FCommerceCommon::QueryOffers(FCommerceQueryOffers::Params&& Params)
+{
+	TOnlineAsyncOpRef<FCommerceQueryOffers> Op = GetOp<FCommerceQueryOffers>(MoveTemp(Params));
+	Op->SetError(Errors::NotImplemented());
+	return Op->GetHandle();
+}
+
+TOnlineAsyncOpHandle<FCommerceQueryOffersById> FCommerceCommon::QueryOffersById(FCommerceQueryOffersById::Params&& Params)
+{
+	TOnlineAsyncOpRef<FCommerceQueryOffersById> Op = GetOp<FCommerceQueryOffersById>(MoveTemp(Params));
+	Op->SetError(Errors::NotImplemented());
+	return Op->GetHandle();
+}
+
+TOnlineResult<FCommerceGetOffers> FCommerceCommon::GetOffers(FCommerceGetOffers::Params&& Params)
+{
+	return TOnlineResult<FCommerceGetOffers>(Errors::NotImplemented());
+}
+
+TOnlineResult<FCommerceGetOffersById> FCommerceCommon::GetOffersById(FCommerceGetOffersById::Params&& Params)
+{
+	return TOnlineResult<FCommerceGetOffersById>(Errors::NotImplemented());
+}
+
+TOnlineAsyncOpHandle<FCommerceShowStoreUI> FCommerceCommon::ShowStoreUI(FCommerceShowStoreUI::Params&& Params)
+{
+	TOnlineAsyncOpRef<FCommerceShowStoreUI> Op = GetOp<FCommerceShowStoreUI>(MoveTemp(Params));
+	Op->SetError(Errors::NotImplemented());
+	return Op->GetHandle();
+}
+
+TOnlineAsyncOpHandle<FCommerceCheckout> FCommerceCommon::Checkout(FCommerceCheckout::Params&& Params)
+{
+	TOnlineAsyncOpRef<FCommerceCheckout> Op = GetOp<FCommerceCheckout>(MoveTemp(Params));
+	Op->SetError(Errors::NotImplemented());
+	return Op->GetHandle();
+}
+
+TOnlineAsyncOpHandle<FCommerceQueryTransactionEntitlements> FCommerceCommon::QueryTransactionEntitlements(FCommerceQueryTransactionEntitlements::Params&& Params)
+{
+	TOnlineAsyncOpRef<FCommerceQueryTransactionEntitlements> Op = GetOp<FCommerceQueryTransactionEntitlements>(MoveTemp(Params));
+	Op->SetError(Errors::NotImplemented());
+	return Op->GetHandle();
+}
+
+TOnlineAsyncOpHandle<FCommerceQueryEntitlements> FCommerceCommon::QueryEntitlements(FCommerceQueryEntitlements::Params&& Params)
+{
+	TOnlineAsyncOpRef<FCommerceQueryEntitlements> Op = GetOp<FCommerceQueryEntitlements>(MoveTemp(Params));
+	Op->SetError(Errors::NotImplemented());
+	return Op->GetHandle();
+}
+
+TOnlineResult<FCommerceGetEntitlements> FCommerceCommon::GetEntitlements(FCommerceGetEntitlements::Params&& Params)
+{
+	return TOnlineResult<FCommerceGetEntitlements>(Errors::NotImplemented());
+}
+
+TOnlineAsyncOpHandle<FCommerceRedeemEntitlement> FCommerceCommon::RedeemEntitlement(FCommerceRedeemEntitlement::Params&& Params)
+{
+	TOnlineAsyncOpRef<FCommerceRedeemEntitlement> Op = GetOp<FCommerceRedeemEntitlement>(MoveTemp(Params));
+	Op->SetError(Errors::NotImplemented());
+	return Op->GetHandle();
+}
+
+TOnlineAsyncOpHandle<FCommerceRetrieveS2SToken> FCommerceCommon::RetrieveS2SToken(FCommerceRetrieveS2SToken::Params&& Params)
+{
+	TOnlineAsyncOpRef<FCommerceRetrieveS2SToken> Op = GetOp<FCommerceRetrieveS2SToken>(MoveTemp(Params));
+	Op->SetError(Errors::NotImplemented());
+	return Op->GetHandle();
+}
+
+TOnlineEvent<void(const FCommerceOnPurchaseComplete&)> FCommerceCommon::OnPurchaseCompleted()
+{
+	return OnPurchaseCompletedEvent;
+}
+
+} // namespace UE::Online
diff --git a/Engine/Plugins/Online/Experimental/OnlineServices/Source/OnlineServicesCommon/Private/Online/OnlineServicesCommon.cpp b/Engine/Plugins/Online/Experimental/OnlineServices/Source/OnlineServicesCommon/Private/Online/OnlineServicesCommon.cpp
--- a/Engine/Plugins/Online/Experimental/OnlineServices/Source/OnlineServicesCommon/Private/Online/OnlineServicesCommon.cpp
+++ b/Engine/Plugins/Online/Experimental/OnlineServices/Source/OnlineServicesCommon/Private/Online/OnlineServicesCommon.cpp
@@ -45,6 +45,12 @@ IAchievementsPtr FOnlineServicesCommon::GetAchievementsInterface()
 	return IAchievementsPtr(AsShared(), Get<IAchievements>());
 }
 
+ICommercePtr FOnlineServicesCommon::GetCommerceInterface()
+{
+	return ICommercePtr(AsShared(), Get<ICommerce>());
+}
+
+
 IAuthPtr FOnlineServicesCommon::GetAuthInterface()
 {
 	return IAuthPtr(AsShared(), Get<IAuth>());
diff --git a/Engine/Plugins/Online/Experimental/OnlineServices/Source/OnlineServicesCommon/Public/Online/CommerceCommon.h b/Engine/Plugins/Online/Experimental/OnlineServices/Source/OnlineServicesCommon/Public/Online/CommerceCommon.h
new file mode 100644
--- /dev/null
+++ b/Engine/Plugins/Online/Experimental/OnlineServices/Source/OnlineServicesCommon/Public/Online/CommerceCommon.h
@@ -0,0 +1,41 @@
+// Copyright Epic Games, Inc. All Rights Reserved.
+
+#pragma once
+
+#include "Online/Commerce.h"
+#include "Online/OnlineComponent.h"
+
+namespace UE::Online {
+
+class FOnlineServicesCommon;
+
+class ONLINESERVICESCOMMON_API FCommerceCommon : public TOnlineComponent<ICommerce>
+{
+public:
+	using Super = ICommerce;
+
+	FCommerceCommon(FOnlineServicesCommon& InServices);
+
+	// IOnlineComponent
+	virtual void RegisterCommands() override;
+
+	// ICommerce
+	virtual TOnlineAsyncOpHandle<FCommerceQueryOffers> QueryOffers(FCommerceQueryOffers::Params&& Params) override;
+	virtual TOnlineAsyncOpHandle<FCommerceQueryOffersById> QueryOffersById(FCommerceQueryOffersById::Params&& Params) override;
+	virtual TOnlineResult<FCommerceGetOffers> GetOffers(FCommerceGetOffers::Params&& Params) override;
+	virtual TOnlineResult<FCommerceGetOffersById> GetOffersById(FCommerceGetOffersById::Params&& Params) override;
+	virtual TOnlineAsyncOpHandle<FCommerceShowStoreUI> ShowStoreUI(FCommerceShowStoreUI::Params&& Params) override;
+	virtual TOnlineAsyncOpHandle<FCommerceCheckout> Checkout(FCommerceCheckout::Params&& Params) override;
+	virtual TOnlineAsyncOpHandle<FCommerceQueryTransactionEntitlements> QueryTransactionEntitlements(FCommerceQueryTransactionEntitlements::Params&& Params) override;
+	virtual TOnlineAsyncOpHandle<FCommerceQueryEntitlements> QueryEntitlements(FCommerceQueryEntitlements::Params&& Params) override;
+	virtual TOnlineResult<FCommerceGetEntitlements> GetEntitlements(FCommerceGetEntitlements::Params&& Params) override;
+	virtual TOnlineAsyncOpHandle<FCommerceRedeemEntitlement> RedeemEntitlement(FCommerceRedeemEntitlement::Params&& Params) override;
+	virtual TOnlineAsyncOpHandle<FCommerceRetrieveS2SToken> RetrieveS2SToken(FCommerceRetrieveS2SToken::Params&& Params) override;
+	virtual TOnlineEvent<void(const FCommerceOnPurchaseComplete&)> OnPurchaseCompleted() override;
+
+protected:
+	TOnlineEventCallable<void(const FCommerceOnPurchaseComplete&)> OnPurchaseCompletedEvent;
+};
+
+/* UE::Online */
+}
diff --git a/Engine/Plugins/Online/Experimental/OnlineServices/Source/OnlineServicesCommon/Public/Online/OnlineServicesCommon.h b/Engine/Plugins/Online/Experimental/OnlineServices/Source/OnlineServicesCommon/Public/Online/OnlineServicesCommon.h
--- a/Engine/Plugins/Online/Experimental/OnlineServices/Source/OnlineServicesCommon/Public/Online/OnlineServicesCommon.h
+++ b/Engine/Plugins/Online/Experimental/OnlineServices/Source/OnlineServicesCommon/Public/Online/OnlineServicesCommon.h
@@ -38,6 +38,7 @@ public:
 	virtual IAchievementsPtr GetAchievementsInterface() override;
 	virtual IAuthPtr GetAuthInterface() override;
 	virtual IUserInfoPtr GetUserInfoInterface() override;
+	virtual ICommercePtr GetCommerceInterface() override;
 	virtual IFriendsPtr GetFriendsInterface() override;
 	virtual IPresencePtr GetPresenceInterface() override;
 	virtual IExternalUIPtr GetExternalUIInterface() override;
diff --git a/Engine/Plugins/Online/Experimental/OnlineServices/Source/OnlineServicesInterface/Private/Online/Commerce.cpp b/Engine/Plugins/Online/Experimental/OnlineServices/Source/OnlineServicesInterface/Private/Online/Commerce.cpp
new file mode 100644
--- /dev/null
+++ b/Engine/Plugins/Online/Experimental/OnlineServices/Source/OnlineServicesInterface/Private/Online/Commerce.cpp
@@ -0,0 +1,44 @@
+// Copyright Epic Games, Inc. All Rights Reserved.
+
+#include "Online/Commerce.h"
+
+namespace UE::Online {
+
+const TCHAR* LexToString(EOfferType OfferType)
+{
+	switch (OfferType)
+	{
+	case EOfferType::Nonconsumable: return TEXT("Nonconsumable");
+	case EOfferType::Consumable:return TEXT("Consumable");
+	case EOfferType::Subscription: return TEXT("Subscription");
+	default: checkNoEntry();
+	case EOfferType::Unknown: return TEXT("Unknown");
+	}
+}
+
+void LexFromString(EOfferType& OutOfferType, const TCHAR* InStr)
+{
+	if (FCString::Stricmp(InStr, TEXT("Nonconsumable")) == 0)
+	{
+		OutOfferType = EOfferType::Nonconsumable;
+	}
+	else if (FCString::Stricmp(InStr, TEXT("Consumable")) == 0)
+	{
+		OutOfferType = EOfferType::Consumable;
+	}
+	else if (FCString::Stricmp(InStr, TEXT("Subscription")) == 0)
+	{
+		OutOfferType = EOfferType::Subscription;
+	}
+	else if (FCString::Stricmp(InStr, TEXT("Unknown")) == 0)
+	{
+		OutOfferType = EOfferType::Unknown;
+	}
+	else
+	{
+		checkNoEntry();
+		OutOfferType = EOfferType::Unknown;
+	}
+}
+
+} // namespace UE::Online
\ No newline at end of file
diff --git a/Engine/Plugins/Online/Experimental/OnlineServices/Source/OnlineServicesInterface/Public/Online/Commerce.h b/Engine/Plugins/Online/Experimental/OnlineServices/Source/OnlineServicesInterface/Public/Online/Commerce.h
new file mode 100644
--- /dev/null
+++ b/Engine/Plugins/Online/Experimental/OnlineServices/Source/OnlineServicesInterface/Public/Online/Commerce.h
@@ -0,0 +1,478 @@
+// Copyright Epic Games, Inc. All Rights Reserved.
+
+#pragma once
+
+#include "Online/CoreOnline.h"
+#include "Online/OnlineAsyncOpHandle.h"
+#include "Online/OnlineMeta.h"
+
+namespace UE::Online {
+
+/* Unique identifier for an offer */
+using FOfferId = FString;
+
+/* Unique identifier for an entitlement */
+using FEntitlementId = FString;
+	
+enum class EOfferType : uint8
+{
+	Nonconsumable,
+	Consumable,
+	Subscription,
+	Unknown
+};
+
+ONLINESERVICESINTERFACE_API const TCHAR* LexToString(EOfferType Status);
+ONLINESERVICESINTERFACE_API void LexFromString(EOfferType& OutStatus, const TCHAR* InStr);
+
+// Contains all the information required to display an offer to a user.
+struct FOffer
+{
+	/* The Id of this offer */
+	FOfferId OfferId;
+	/* Localized title of this offer */
+	FText Title;
+	/* Localized description of this offer */
+	FText Description;
+	/* Localized long description of this offer */
+	FText LongDescription;
+	/* The maximum number of times the offer can be purchased (-1 for no limit) */
+	int32 PurchaseLimit;
+	/* Currency Code for this offer */
+	FString CurrencyCode;
+	/* Regular non-sale price as text for display */
+	FText FormattedRegularPrice;
+	/* Regular non-sale price in numeric form for comparison/sorting */
+	uint64 RegularPrice;
+	/* Final-Pricing (Post-Sales/Discounts) as text for display */
+	FText FormattedPrice;
+	/* Final-Price (Post-Sales/Discounts) in numeric form for comparison/sorting */
+	uint64 Price;
+	/* Number of decimal places in Price/RegularPrice */
+	int32 PriceDecimalPoint;
+	/* Date the offer was first available/will be available */
+	TOptional<FDateTime> ReleaseDate;
+	/* Date the information is no longer valid (offer no longer available or discount expires) */
+	TOptional<FDateTime> ExpirationDate;
+	/* Additional implementation specified fields. See the implementation specific documentation for what this will contain */
+	TMap<FString, FString> AdditionalData;
+	// TODO: Expose offer image data? additional subscription related fields? Offer type (consumable, subscription, dlc, etc)?
+};
+
+struct FCommerceQueryOffers
+{
+	static constexpr TCHAR Name[] = TEXT("QueryOffers");
+	
+	struct Params
+	{
+		/* Local user */
+		FOnlineAccountIdHandle LocalUserId;
+	};
+
+	struct Result
+	{
+
+	};
+};
+
+// Query the offers matching specified offer ids
+struct FCommerceQueryOffersById
+{
+	static constexpr TCHAR Name[] = TEXT("QueryOffersById");
+
+	struct Params
+	{
+		/* Local user */
+		FOnlineAccountIdHandle LocalUserId;
+		/* List of offer ids */
+		TArray<FOfferId> OfferIds;
+	};
+
+	struct Result
+	{
+		/* List of offer ids that were found */
+		TArray<FOfferId> OfferIds;
+	};
+};
+
+// Get the complete list of cached offers
+struct FCommerceGetOffers
+{
+	static constexpr TCHAR Name[] = TEXT("GetOffers");
+	
+	struct Params
+	{
+		/* Local user */
+		FOnlineAccountIdHandle LocalUserId;
+	};
+
+	struct Result
+	{
+		/* List of offers */
+		TArray<FOffer> Offers;
+	};
+};
+
+// Get the offers for a list of offer ids
+struct FCommerceGetOffersById
+{
+	static constexpr TCHAR Name[] = TEXT("GetOffersById");
+		
+	struct Params
+	{
+		/* Local user */
+		FOnlineAccountIdHandle LocalUserId;
+		/* List of offer ids */
+		TArray<FOfferId> OfferIds;
+	};
+
+	struct Result
+	{
+		/* List of offers */
+		TArray<FOffer> Offers;
+	};
+};
+
+//Bring up the platform store ui
+struct FCommerceShowStoreUI
+{
+
+	static constexpr TCHAR Name[] = TEXT("ShowStoreUI");
+
+	struct Params
+	{
+		/* Local user */
+		FOnlineAccountIdHandle LocalUserId;
+	};
+
+	struct Result
+	{
+
+	};
+};
+
+struct FPurchaseOffer
+{
+	/* The Id of the offer */
+	FOfferId OfferId;
+	/* Quantity of the offer to purchase */
+	int32 Quantity;
+};
+
+// Initiate the checkout process for purchasing one or more offers
+struct FCommerceCheckout
+{
+	static constexpr TCHAR Name[] = TEXT("Checkout");
+		
+	struct Params
+	{
+		/* Local user */
+		FOnlineAccountIdHandle LocalUserId;
+		/* Offers to purchase */
+		TArray<FPurchaseOffer> Offers;
+	};
+
+	struct Result
+	{
+		/* If available, the transaction id for the completed purchase */
+		TOptional<FString> TransactionId;
+	};
+};
+
+struct FEntitlement
+{
+	/* Id of the entitlement */
+	FEntitlementId EntitlementId;
+	/* Type of the entitlement */
+	FName EntitlementType;
+	/* The Id of the product */
+	FString ProductId;
+	/* True if the entitlement has been marked as redeemed */
+	bool bRedeemed;
+	/* Quantity */
+	int32 Quantity;
+	/* If set, specifies when the entitlement was acquired */
+	TOptional<FDateTime> AcquiredDate;
+	/* If set, specifies when the entitlement will expire */
+	TOptional<FDateTime> ExpiryDate;
+};
+
+// Retrieve the list of entitlements from a transaction id returned from a checkout call
+struct FCommerceQueryTransactionEntitlements
+{
+
+	static constexpr TCHAR Name[] = TEXT("QueryTransactionEntitlements");
+		
+	struct Params
+	{
+		/* Local user */
+		FOnlineAccountIdHandle LocalUserId;
+		/* Transaction id for a completed purchase */
+		FString TransactionId;
+	};
+
+	struct Result
+	{
+		/* Array of entitlements granted in the transaction */
+		TArray<FEntitlement> Entitlements;
+	};
+};
+
+// Get the list of all entitlements
+struct FCommerceQueryEntitlements
+{
+	static constexpr TCHAR Name[] = TEXT("QueryEntitlements");
+		
+	struct Params
+	{
+		/* Local user */
+		FOnlineAccountIdHandle LocalUserId;
+		/* Whether or not to fetch redeemed entitlements */
+		bool bIncludeRedeemed;
+	};
+
+	struct Result
+	{
+		/* Empty, complete entitlement list can be retrieved via GetEntitlements */
+	};
+};
+
+// Get the complete list of cached entitlements
+struct FCommerceGetEntitlements
+{
+	static constexpr TCHAR Name[] = TEXT("GetEntitlements");
+		
+	struct Params
+	{
+		/* Local user */
+		FOnlineAccountIdHandle LocalUserId;
+	};
+
+	struct Result
+	{
+		/* Array of entitlements owned by the user */
+		TArray<FEntitlement> Entitlements;
+	};
+};
+
+// Mark an entitlement as redeemed.Used when there is no external service managing redemption of entitlements or the redemption must be done on the client side.
+struct FCommerceRedeemEntitlement
+{
+	static constexpr TCHAR Name[] = TEXT("RedeemEntitlement");
+		
+	struct Params
+	{
+		/* Local user */
+		FOnlineAccountIdHandle LocalUserId;
+		/* Entitlement id to mark as redeemed */
+		FEntitlementId EntitlementId;
+		/* Quantity to consume */
+		int32 Quantity;
+	};
+
+	struct Result
+	{
+
+	};
+};
+
+// Get a token suitable for making service to service Commerce related calls
+struct FCommerceRetrieveS2SToken
+{
+	static constexpr TCHAR Name[] = TEXT("RetrieveS2SToken");
+		
+	struct Params
+	{
+		/* Local user */
+		FOnlineAccountIdHandle LocalUserId;
+		/* Implementation specific, used to specify the type of token to retrieve. See implementation specific documentation for valid values */
+		FString TokenType;
+	};
+
+	struct Result
+	{
+		/* Token or code that can be used for service to service calls. See implementation specific documentation for details */
+		TArray<uint8> Token;
+	};
+};
+
+// event
+struct FCommerceOnPurchaseComplete
+{
+public:
+	
+	/* Local user */
+	FOnlineAccountIdHandle LocalUserId;
+	
+	/* If available, the transaction id for the completed purchase */
+	TOptional<FString> TransactionId;
+
+};
+
+class ICommerce
+{
+	/* Query the list of all offers for the target user */
+	virtual TOnlineAsyncOpHandle<FCommerceQueryOffers> QueryOffers(FCommerceQueryOffers::Params&& Params) = 0;
+
+	/* Query the offers matching specified offer ids */
+	virtual TOnlineAsyncOpHandle<FCommerceQueryOffersById> QueryOffersById(FCommerceQueryOffersById::Params&& Params) = 0;
+
+	/* Get the complete list of cached offers */
+	virtual TOnlineResult<FCommerceGetOffers> GetOffers(FCommerceGetOffers::Params&& Params) = 0;
+
+	/* Get the offers for a list of offer ids */
+	virtual TOnlineResult<FCommerceGetOffersById> GetOffersById(FCommerceGetOffersById::Params&& Params) = 0;
+
+	/* Bring up the platform store ui */
+	virtual TOnlineAsyncOpHandle<FCommerceShowStoreUI> ShowStoreUI(FCommerceShowStoreUI::Params&& Params) = 0;
+
+	/* Initiate the checkout process for purchasing one or more offers */
+	virtual TOnlineAsyncOpHandle<FCommerceCheckout> Checkout(FCommerceCheckout::Params&& Params) = 0;
+
+	/* Retrieve the list of entitlements from a transaction id returned from a checkout call */
+	virtual TOnlineAsyncOpHandle<FCommerceQueryTransactionEntitlements> QueryTransactionEntitlements(FCommerceQueryTransactionEntitlements::Params&& Params) = 0;
+
+	/* Get the list of all entitlements */
+	virtual TOnlineAsyncOpHandle<FCommerceQueryEntitlements> QueryEntitlements(FCommerceQueryEntitlements::Params&& Params) = 0;
+
+	/* Get the complete list of cached entitlements */
+	virtual TOnlineResult<FCommerceGetEntitlements> GetEntitlements(FCommerceGetEntitlements::Params&& Params) = 0;
+
+	/* Mark an entitlement as redeemed. Used when there is no external service managing redemption of entitlements or the redemption must be done on the client side */
+	virtual TOnlineAsyncOpHandle<FCommerceRedeemEntitlement> RedeemEntitlement(FCommerceRedeemEntitlement::Params&& Params) = 0;
+
+	/* Get a token suitable for making service to service Commerce related calls */
+	virtual TOnlineAsyncOpHandle<FCommerceRetrieveS2SToken> RetrieveS2SToken(FCommerceRetrieveS2SToken::Params&& Params) = 0;
+
+	/* Fires whenever a purchase is completed */
+	virtual TOnlineEvent<void(const FCommerceOnPurchaseComplete&)> OnPurchaseCompleted() = 0;
+};
+
+namespace Meta {
+
+BEGIN_ONLINE_STRUCT_META(FOffer)
+	ONLINE_STRUCT_FIELD(FOffer, OfferId),
+	ONLINE_STRUCT_FIELD(FOffer, Title),
+	ONLINE_STRUCT_FIELD(FOffer, Description),
+	ONLINE_STRUCT_FIELD(FOffer, LongDescription),
+	ONLINE_STRUCT_FIELD(FOffer, PurchaseLimit),
+	ONLINE_STRUCT_FIELD(FOffer, CurrencyCode),
+	ONLINE_STRUCT_FIELD(FOffer, FormattedRegularPrice),
+	ONLINE_STRUCT_FIELD(FOffer, RegularPrice),
+	ONLINE_STRUCT_FIELD(FOffer, FormattedPrice),
+	ONLINE_STRUCT_FIELD(FOffer, Price),
+	ONLINE_STRUCT_FIELD(FOffer, PriceDecimalPoint),
+	ONLINE_STRUCT_FIELD(FOffer, ReleaseDate),
+	ONLINE_STRUCT_FIELD(FOffer, ExpirationDate),
+	ONLINE_STRUCT_FIELD(FOffer, AdditionalData)
+END_ONLINE_STRUCT_META()
+
+
+BEGIN_ONLINE_STRUCT_META(FPurchaseOffer)
+	ONLINE_STRUCT_FIELD(FPurchaseOffer, OfferId),
+	ONLINE_STRUCT_FIELD(FPurchaseOffer, Quantity)
+END_ONLINE_STRUCT_META()
+
+BEGIN_ONLINE_STRUCT_META(FEntitlement)
+	ONLINE_STRUCT_FIELD(FEntitlement, EntitlementId),
+	ONLINE_STRUCT_FIELD(FEntitlement, EntitlementType),
+	ONLINE_STRUCT_FIELD(FEntitlement, ProductId),
+	ONLINE_STRUCT_FIELD(FEntitlement, bRedeemed),
+	ONLINE_STRUCT_FIELD(FEntitlement, Quantity),
+	ONLINE_STRUCT_FIELD(FEntitlement, AcquiredDate),
+	ONLINE_STRUCT_FIELD(FEntitlement, ExpiryDate)
+END_ONLINE_STRUCT_META()
+
+BEGIN_ONLINE_STRUCT_META(FCommerceQueryOffers::Params)
+	ONLINE_STRUCT_FIELD(FCommerceQueryOffers::Params, LocalUserId)
+END_ONLINE_STRUCT_META()
+	BEGIN_ONLINE_STRUCT_META(FCommerceQueryOffers::Result)
+END_ONLINE_STRUCT_META()
+
+BEGIN_ONLINE_STRUCT_META(FCommerceQueryOffersById::Params)
+	ONLINE_STRUCT_FIELD(FCommerceQueryOffersById::Params, LocalUserId),
+	ONLINE_STRUCT_FIELD(FCommerceQueryOffersById::Params, OfferIds)
+END_ONLINE_STRUCT_META()
+BEGIN_ONLINE_STRUCT_META(FCommerceQueryOffersById::Result)
+	ONLINE_STRUCT_FIELD(FCommerceQueryOffersById::Result, OfferIds)
+END_ONLINE_STRUCT_META()
+
+BEGIN_ONLINE_STRUCT_META(FCommerceGetOffers::Params)
+	ONLINE_STRUCT_FIELD(FCommerceGetOffers::Params, LocalUserId)
+END_ONLINE_STRUCT_META()
+BEGIN_ONLINE_STRUCT_META(FCommerceGetOffers::Result)
+	ONLINE_STRUCT_FIELD(FCommerceGetOffers::Result, Offers)
+END_ONLINE_STRUCT_META()
+
+BEGIN_ONLINE_STRUCT_META(FCommerceGetOffersById::Params)
+	ONLINE_STRUCT_FIELD(FCommerceGetOffersById::Params, LocalUserId),
+	ONLINE_STRUCT_FIELD(FCommerceGetOffersById::Params, OfferIds)
+END_ONLINE_STRUCT_META()
+BEGIN_ONLINE_STRUCT_META(FCommerceGetOffersById::Result)
+ONLINE_STRUCT_FIELD(FCommerceGetOffersById::Result, Offers)
+END_ONLINE_STRUCT_META()
+
+BEGIN_ONLINE_STRUCT_META(FCommerceShowStoreUI::Params)
+	ONLINE_STRUCT_FIELD(FCommerceShowStoreUI::Params, LocalUserId)
+END_ONLINE_STRUCT_META()
+BEGIN_ONLINE_STRUCT_META(FCommerceShowStoreUI::Result)
+END_ONLINE_STRUCT_META()
+
+BEGIN_ONLINE_STRUCT_META(FCommerceCheckout::Params)
+	ONLINE_STRUCT_FIELD(FCommerceCheckout::Params, LocalUserId),
+	ONLINE_STRUCT_FIELD(FCommerceCheckout::Params, Offers)
+END_ONLINE_STRUCT_META()
+BEGIN_ONLINE_STRUCT_META(FCommerceCheckout::Result)
+	ONLINE_STRUCT_FIELD(FCommerceCheckout::Result, TransactionId)
+END_ONLINE_STRUCT_META()
+
+BEGIN_ONLINE_STRUCT_META(FCommerceQueryTransactionEntitlements::Params)
+	ONLINE_STRUCT_FIELD(FCommerceQueryTransactionEntitlements::Params, LocalUserId),
+	ONLINE_STRUCT_FIELD(FCommerceQueryTransactionEntitlements::Params, TransactionId)
+END_ONLINE_STRUCT_META()
+BEGIN_ONLINE_STRUCT_META(FCommerceQueryTransactionEntitlements::Result)
+	ONLINE_STRUCT_FIELD(FCommerceQueryTransactionEntitlements::Result, Entitlements)
+END_ONLINE_STRUCT_META()
+
+BEGIN_ONLINE_STRUCT_META(FCommerceQueryEntitlements::Params)
+	ONLINE_STRUCT_FIELD(FCommerceQueryEntitlements::Params, LocalUserId),
+	ONLINE_STRUCT_FIELD(FCommerceQueryEntitlements::Params, bIncludeRedeemed)
+END_ONLINE_STRUCT_META()
+BEGIN_ONLINE_STRUCT_META(FCommerceQueryEntitlements::Result)
+END_ONLINE_STRUCT_META()
+
+BEGIN_ONLINE_STRUCT_META(FCommerceGetEntitlements::Params)
+	ONLINE_STRUCT_FIELD(FCommerceGetEntitlements::Params, LocalUserId)
+END_ONLINE_STRUCT_META()
+BEGIN_ONLINE_STRUCT_META(FCommerceGetEntitlements::Result)
+	ONLINE_STRUCT_FIELD(FCommerceGetEntitlements::Result, Entitlements)
+END_ONLINE_STRUCT_META()
+
+BEGIN_ONLINE_STRUCT_META(FCommerceRedeemEntitlement::Params)
+	ONLINE_STRUCT_FIELD(FCommerceRedeemEntitlement::Params, LocalUserId),
+	ONLINE_STRUCT_FIELD(FCommerceRedeemEntitlement::Params, EntitlementId),
+	ONLINE_STRUCT_FIELD(FCommerceRedeemEntitlement::Params, Quantity)
+END_ONLINE_STRUCT_META()
+BEGIN_ONLINE_STRUCT_META(FCommerceRedeemEntitlement::Result)
+END_ONLINE_STRUCT_META()
+
+BEGIN_ONLINE_STRUCT_META(FCommerceRetrieveS2SToken::Params)
+	ONLINE_STRUCT_FIELD(FCommerceRetrieveS2SToken::Params, LocalUserId),
+	ONLINE_STRUCT_FIELD(FCommerceRetrieveS2SToken::Params, TokenType)
+END_ONLINE_STRUCT_META()
+BEGIN_ONLINE_STRUCT_META(FCommerceRetrieveS2SToken::Result)
+	ONLINE_STRUCT_FIELD(FCommerceRetrieveS2SToken::Result, Token)
+
+END_ONLINE_STRUCT_META()
+
+BEGIN_ONLINE_STRUCT_META(FCommerceOnPurchaseComplete)
+	ONLINE_STRUCT_FIELD(FCommerceOnPurchaseComplete, LocalUserId),
+	ONLINE_STRUCT_FIELD(FCommerceOnPurchaseComplete, TransactionId)
+END_ONLINE_STRUCT_META()
+
+
+} // namespace Meta
+} // namespace UE::Online
+
+
+
diff --git a/Engine/Plugins/Online/Experimental/OnlineServices/Source/OnlineServicesInterface/Public/Online/OnlineServices.h b/Engine/Plugins/Online/Experimental/OnlineServices/Source/OnlineServicesInterface/Public/Online/OnlineServices.h
--- a/Engine/Plugins/Online/Experimental/OnlineServices/Source/OnlineServicesInterface/Public/Online/OnlineServices.h
+++ b/Engine/Plugins/Online/Experimental/OnlineServices/Source/OnlineServicesInterface/Public/Online/OnlineServices.h
@@ -15,6 +15,7 @@ namespace UE::Online {
 using IOnlineServicesPtr = TSharedPtr<class IOnlineServices>;
 using IAchievementsPtr = TSharedPtr<class IAchievements>;
 using IAuthPtr = TSharedPtr<class IAuth>;
+using ICommercePtr = TSharedPtr<class ICommerce>;
 using IUserInfoPtr = TSharedPtr<class IUserInfo>;
 using IFriendsPtr = TSharedPtr<class IFriends>;
 using IPresencePtr = TSharedPtr<class IPresence>;
@@ -69,6 +70,11 @@ public:
 	 */
 	virtual IAuthPtr GetAuthInterface() = 0;
 
+	/**
+	 *
+	 */
+	virtual ICommercePtr GetCommerceInterface() = 0;
+
 	/**
 	 *
 	 */
diff --git a/Engine/Plugins/Online/Experimental/OnlineServices/Source/OnlineServicesInterface/Public/Online/OnlineUtils.h b/Engine/Plugins/Online/Experimental/OnlineServices/Source/OnlineServicesInterface/Public/Online/OnlineUtils.h
--- a/Engine/Plugins/Online/Experimental/OnlineServices/Source/OnlineServicesInterface/Public/Online/OnlineUtils.h
+++ b/Engine/Plugins/Online/Experimental/OnlineServices/Source/OnlineServicesInterface/Public/Online/OnlineUtils.h
@@ -19,6 +19,8 @@ template <typename... Ts> FString ToLogString(const TVariant<Ts...>& Variant);
 template <typename T, ESPMode Mode> FString ToLogString(const TSharedRef<T, Mode>& Ref);
 inline FString ToLogString(const FString& String);
 inline FString ToLogString(const FName& Name);
+inline FString ToLogString(const FText& Text);
+inline FString ToLogString(const FDateTime& Time);
 inline FString ToLogString(uint8 Value);
 inline FString ToLogString(int8 Value);
 inline FString ToLogString(uint16 Value);
@@ -108,6 +110,17 @@ inline FString ToLogString(const FName& Name)
 	return Name.ToString();
 }
 
+inline FString ToLogString(const FText& Text)
+{
+	return Text.ToString();
+}
+
+inline FString ToLogString(const FDateTime& Time)
+{
+	return Time.ToString();
+}
+
+
 inline FString ToLogString(uint8 Value)
 {
 	return FString::Printf(TEXT("%hhu"), Value);
