Default service override support

diff --git a/Engine/Plugins/Online/Experimental/OnlineServices/Source/OnlineServicesInterface/Private/Online/OnlineServicesRegistry.cpp b/Engine/Plugins/Online/Experimental/OnlineServices/Source/OnlineServicesInterface/Private/Online/OnlineServicesRegistry.cpp
--- a/Engine/Plugins/Online/Experimental/OnlineServices/Source/OnlineServicesInterface/Private/Online/OnlineServicesRegistry.cpp
+++ b/Engine/Plugins/Online/Experimental/OnlineServices/Source/OnlineServicesInterface/Private/Online/OnlineServicesRegistry.cpp
@@ -52,9 +52,17 @@ TSharedPtr<IOnlineServices> FOnlineServicesRegistry::GetNamedServicesInstance(EO
 
 	if (OnlineServices == EOnlineServices::Default)
 	{
-		FString Value;
-		GConfig->GetString(TEXT("OnlineServices"), TEXT("DefaultServices"), Value, GEngineIni);
-		LexFromString(OnlineServices, *Value);
+		if(DefaultServiceOverride != EOnlineServices::Default)
+		{
+			OnlineServices = DefaultServiceOverride;
+		}
+		else
+		{
+			FString Value;
+			GConfig->GetString(TEXT("OnlineServices"), TEXT("DefaultServices"), Value, GEngineIni);
+
+			LexFromString(OnlineServices, *Value);
+		};
 	}
 	else if (OnlineServices == EOnlineServices::Platform)
 	{
@@ -83,6 +91,18 @@ TSharedPtr<IOnlineServices> FOnlineServicesRegistry::GetNamedServicesInstance(EO
 	return Services;
 }
 
+#if WITH_DEV_AUTOMATION_TESTS
+void FOnlineServicesRegistry::SetDefaultServiceOverride(EOnlineServices DefaultService)
+{
+	DefaultServiceOverride = DefaultService;
+}
+
+void FOnlineServicesRegistry::ClearDefaultServiceOverride()
+{
+	DefaultServiceOverride = EOnlineServices::Default;
+}
+#endif //WITH_DEV_AUTOMATION_TESTS
+
 void FOnlineServicesRegistry::DestroyNamedServicesInstance(EOnlineServices OnlineServices, FName InstanceName)
 {
 	if (TSharedRef<IOnlineServices>* ServicesPtr = NamedServiceInstances.FindOrAdd(OnlineServices).Find(InstanceName))
diff --git a/Engine/Plugins/Online/Experimental/OnlineServices/Source/OnlineServicesInterface/Public/Online/OnlineServicesRegistry.h b/Engine/Plugins/Online/Experimental/OnlineServices/Source/OnlineServicesInterface/Public/Online/OnlineServicesRegistry.h
--- a/Engine/Plugins/Online/Experimental/OnlineServices/Source/OnlineServicesInterface/Public/Online/OnlineServicesRegistry.h
+++ b/Engine/Plugins/Online/Experimental/OnlineServices/Source/OnlineServicesInterface/Public/Online/OnlineServicesRegistry.h
@@ -105,6 +105,20 @@ public:
 	 * @param OutOnlineServices Array of online services to fill
 	 */
 	ONLINESERVICESINTERFACE_API void GetAllServicesInstances(TArray<TSharedRef<IOnlineServices>>& OutOnlineServices) const;
+
+#if WITH_DEV_AUTOMATION_TESTS
+	/**
+	 * Adds a temporary override to the default online service. Used for testing to quickly run tests with different default online services.
+	 * 
+	 * @param Service the name of the Online Service to force default online service to return
+	 */
+	ONLINESERVICESINTERFACE_API void SetDefaultServiceOverride(EOnlineServices DefaultService);
+
+	/**
+	 * Removes an override for the target service.
+	 */
+	ONLINESERVICESINTERFACE_API void ClearDefaultServiceOverride();
+#endif //WITH_DEV_AUTOMATION_TESTS
 private:
 	struct FFactoryAndPriority
 	{
@@ -120,6 +134,7 @@ private:
 
 	TMap<EOnlineServices, FFactoryAndPriority> ServicesFactories;
 	TMap<EOnlineServices, TMap<FName, TSharedRef<IOnlineServices>>> NamedServiceInstances;
+	EOnlineServices DefaultServiceOverride;
 
 	FOnlineServicesRegistry() {}
 	~FOnlineServicesRegistry();
