Implement OSSv2 online schema

diff --git a/Engine/Plugins/Online/Experimental/OnlineServices/Source/OnlineServicesCommon/Public/Online/OnlineConfig.h b/Engine/Plugins/Online/Experimental/OnlineServices/Source/OnlineServicesCommon/Public/Online/OnlineConfig.h
--- a/Engine/Plugins/Online/Experimental/OnlineServices/Source/OnlineServicesCommon/Public/Online/OnlineConfig.h
+++ b/Engine/Plugins/Online/Experimental/OnlineServices/Source/OnlineServicesCommon/Public/Online/OnlineConfig.h
@@ -4,6 +4,7 @@
 
 #include "Misc/ConfigCacheIni.h"
 #include "Online/OnlineMeta.h"
+#include "Algo/Transform.h"
 
 namespace UE::Online {
 
@@ -73,6 +74,26 @@ public:
 		return false;
 	}
 
+	/**
+	 * Get an FName value
+	 *
+	 * @param Section Section to read the value from
+	 * @param Key Key in the section for the value
+	 * @param Value FName value read from the config. Will be unchanged if not present
+	 *
+	 * @return true if a value was read
+	 */
+	virtual bool GetValue(const TCHAR* Section, const TCHAR* Key, FName& Value)
+	{
+		FString StringValue;
+		if (GetValue(Section, Key, StringValue))
+		{
+			Value = FName(*StringValue);
+			return true;
+		}
+		return false;
+	}
+
 	/**
 	 * Get a bool value
 	 *
@@ -209,12 +230,38 @@ public:
 		if (GetValue(Section, Key, StringArray) > 0)
 		{
 			Value.Reserve(StringArray.Num());
-			for (const FString& StringValue : StringArray)
-			{
-				FText TextValue;
-				FTextStringHelper::ReadFromBuffer(*StringValue, TextValue, Section);
-				Value.Emplace(TextValue);
-			}
+			Algo::Transform(StringArray, Value,
+				[Section](const FString& StringValue) -> FText
+				{
+					FText TextValue;
+					FTextStringHelper::ReadFromBuffer(*StringValue, TextValue, Section);
+					return TextValue;
+				});
+		}
+		return Value.Num();
+	}
+
+	/**
+	 * Get a value consisting of a TArray of FName
+	 *
+	 * @param Section Section to read the value from
+	 * @param Key Key in the section for the value
+	 * @param Value Array of FName values read from the config. Will be empty if not present
+	 *
+	 * @return number of values in the array
+	 */
+	virtual int32 GetValue(const TCHAR* Section, const TCHAR* Key, TArray<FName>& Value)
+	{
+		Value.Empty();
+		TArray<FString> StringArray;
+		if (GetValue(Section, Key, StringArray) > 0)
+		{
+			Value.Reserve(StringArray.Num());
+			Algo::Transform(StringArray, Value,
+				[](const FString& StringValue) -> const TCHAR*
+				{
+					return *StringValue;
+				});
 		}
 		return Value.Num();
 	}
@@ -235,10 +282,11 @@ public:
 		if (GetValue(Section, Key, StringArray) > 0)
 		{
 			Value.Reserve(StringArray.Num());
-			for (const FString& StringValue : StringArray)
-			{
-				Value.Emplace(FCString::ToBool(*StringValue));
-			}
+			Algo::Transform(StringArray, Value,
+				[](const FString& StringValue) -> bool
+				{
+					return FCString::ToBool(*StringValue);
+				});
 		}
 		return Value.Num();
 	}
@@ -259,10 +307,11 @@ public:
 		if (GetValue(Section, Key, StringArray) > 0)
 		{
 			Value.Reserve(StringArray.Num());
-			for (const FString& StringValue : StringArray)
-			{
-				Value.Emplace(FCString::Strtoi(*StringValue, nullptr, 10));
-			}
+			Algo::Transform(StringArray, Value,
+				[](const FString& StringValue) -> int32
+				{
+					return FCString::Strtoi(*StringValue, nullptr, 10);
+				});
 		}
 		return Value.Num();
 	}
@@ -283,10 +332,11 @@ public:
 		if (GetValue(Section, Key, StringArray) > 0)
 		{
 			Value.Reserve(StringArray.Num());
-			for (const FString& StringValue : StringArray)
-			{
-				Value.Emplace(FCString::Strtoi64(*StringValue, nullptr, 10));
-			}
+			Algo::Transform(StringArray, Value,
+				[](const FString& StringValue) -> int64
+				{
+					return FCString::Strtoi64(*StringValue, nullptr, 10);
+				});
 		}
 		return Value.Num();
 	}
@@ -307,10 +357,11 @@ public:
 		if (GetValue(Section, Key, StringArray) > 0)
 		{
 			Value.Reserve(StringArray.Num());
-			for (const FString& StringValue : StringArray)
-			{
-				Value.Emplace(FCString::Strtoui64(*StringValue, nullptr, 10));
-			}
+			Algo::Transform(StringArray, Value,
+				[](const FString& StringValue) -> uint64
+				{
+					return FCString::Strtoui64(*StringValue, nullptr, 10);
+				});
 		}
 		return Value.Num();
 	}
@@ -331,10 +382,11 @@ public:
 		if (GetValue(Section, Key, StringArray) > 0)
 		{
 			Value.Reserve(StringArray.Num());
-			for (const FString& StringValue : StringArray)
-			{
-				Value.Emplace(FCString::Atof(*StringValue));
-			}
+			Algo::Transform(StringArray, Value,
+				[](const FString& StringValue) -> float
+				{
+					return FCString::Atof(*StringValue);
+				});
 		}
 		return Value.Num();
 	}
@@ -355,13 +407,20 @@ public:
 		if (GetValue(Section, Key, StringArray) > 0)
 		{
 			Value.Reserve(StringArray.Num());
-			for (const FString& StringValue : StringArray)
-			{
-				Value.Emplace(FCString::Atod(*StringValue));
-			}
+			Algo::Transform(StringArray, Value,
+				[](const FString& StringValue) -> double
+				{
+					return FCString::Atod(*StringValue);
+				});
 		}
 		return Value.Num();
 	}
+
+	template <typename T>
+	std::enable_if_t<TModels<Meta::COnlineMetadataAvailable, T>::Value, bool> GetValue(const TCHAR* Section, const TCHAR* Key, T& Value);
+
+	template <typename T>
+	std::enable_if_t<TModels<Meta::COnlineMetadataAvailable, T>::Value, int32> GetValue(const TCHAR* Section, const TCHAR* Key, TArray<T>& Value);
 };
 
 /**
@@ -420,6 +479,27 @@ auto LoadConfigValue(IOnlineConfigProvider& Provider, const TCHAR* Section, cons
 	return false;
 }
 
+template <typename T>
+auto LoadConfigValue(IOnlineConfigProvider& Provider, const TCHAR* Section, const TCHAR* Key, TArray<T>& Value)
+	-> std::enable_if_t<std::is_enum_v<T>, int32>
+{
+	Value.Empty();
+	TArray<FString> StringArray;
+	if (Provider.GetValue(Section, Key, StringArray))
+	{
+		Value.Reserve(StringArray.Num());
+		Algo::Transform(StringArray, Value,
+			[](const FString& StringValue) -> T
+			{
+				T EnumValue;
+				using ::LexFromString;
+				LexFromString(EnumValue, *StringValue);
+				return EnumValue;
+			});
+	}
+	return Value.Num();
+}
+
 /* Private */ }
 
 /**
@@ -445,14 +525,66 @@ bool LoadConfig(IOnlineConfigProvider& Provider, const FString& Section, StructT
 }
 
 template <typename StructType>
-bool LoadConfig(IOnlineConfigProvider& Provider, const TArray<FString>& SectionHeiarchy, StructType& Struct)
+bool LoadConfig(IOnlineConfigProvider& Provider, const TArray<FString>& SectionHeiarchy, StructType& OutValue)
 {
 	bool bLoadedValue = false;
 	for (const FString& Section : SectionHeiarchy)
 	{
-		LoadConfig(Provider, Section, Struct);
+		StructType Value;
+		if (LoadConfig(Provider, Section, Value))
+		{
+			bLoadedValue = true;
+			OutValue = MoveTemp(Value);
+		}
+	}
+	return bLoadedValue;
+}
+
+template <typename T>
+std::enable_if_t<TModels<Meta::COnlineMetadataAvailable, T>::Value, bool> IOnlineConfigProvider::GetValue(const TCHAR* Section, const TCHAR* Key, T& OutValue)
+{
+	const FString AggregateSection = FString::Printf(TEXT("%s.%s"), Section, Key);
+	bool bLoadedValue = false;
+	T Value;
+	Meta::VisitFields(Value,
+		[this, &AggregateSection, &bLoadedValue](const TCHAR* FieldName, auto& Field)
+		{
+			bLoadedValue |= Private::LoadConfigValue(*this, *AggregateSection, FieldName, Field);
+		});
+
+	if (bLoadedValue)
+	{
+		OutValue = MoveTemp(Value);
 	}
+
 	return bLoadedValue;
 }
 
+template <typename T>
+std::enable_if_t<TModels<Meta::COnlineMetadataAvailable, T>::Value, int32> IOnlineConfigProvider::GetValue(const TCHAR* Section, const TCHAR* Key, TArray<T>& Value)
+{
+	Value.Empty();
+	TArray<FString> StringArray;
+	if (GetValue(Section, Key, StringArray) > 0)
+	{
+		Value.Reserve(StringArray.Num());
+		for (const FString& StringValue : StringArray)
+		{
+			FString PrecedingToken;
+			const TCHAR* OldValue = *StringValue;
+
+			if (FParse::Token(OldValue, PrecedingToken, true))
+			{
+				const FString AggregateSection = FString::Printf(TEXT("%s.%s.%s"), Section, Key, *StringValue);
+				Meta::VisitFields(Value.Emplace_GetRef(),
+					[this, &AggregateSection](const TCHAR* FieldName, auto& Field)
+					{
+						Private::LoadConfigValue(*this, *AggregateSection, FieldName, Field);
+					});
+			}
+		}
+	}
+	return Value.Num();
+}
+
 /* UE::Online */ }
diff --git a/Engine/Plugins/Online/Experimental/OnlineServices/Source/OnlineServicesInterface/Private/Online/Schema.cpp b/Engine/Plugins/Online/Experimental/OnlineServices/Source/OnlineServicesInterface/Private/Online/Schema.cpp
--- a/Engine/Plugins/Online/Experimental/OnlineServices/Source/OnlineServicesInterface/Private/Online/Schema.cpp
+++ b/Engine/Plugins/Online/Experimental/OnlineServices/Source/OnlineServicesInterface/Private/Online/Schema.cpp
@@ -1,123 +1,502 @@
 // Copyright Epic Games, Inc. All Rights Reserved.
 
 #include "Online/Schema.h"
+#include "Algo/Find.h"
+#include "Algo/ForEach.h"
 
-namespace UE::Online {
+DEFINE_LOG_CATEGORY_STATIC(LogOnlineSchema, Log, All);
 
+namespace UE::Online {
 
-const FString LexToString(const FLexToStringAdaptor<FSchemaVariant>& Adaptor)
+template <typename T>
+T BuildFlagsFromArray(const TArray<T>& Input)
 {
-	return Adaptor.SchemaVariant.GetString();
+	T OutValue = T::None;
+	Algo::ForEach(Input, [&](T FlagValue){ OutValue |= FlagValue; });
+	return OutValue;
 }
 
-void LexFromString(FSchemaVariant& Variant, const TCHAR* InStr)
+bool FSchemaRegistry::ParseConfig(const FSchemaRegistryDescriptorConfig& Config)
 {
-	Variant.Set(TEXT("")); // todo
-}
+	bool ParsedSuccessfully = true;
 
-int64 FSchemaVariant::GetInt64() const
-{
-	int64 AsInt64 = 0;
-	if (VariantData.IsType<int64>())
-	{
-		AsInt64 = VariantData.Get<int64>();
-	}
-	else if (VariantData.IsType<bool>())
-	{
-		AsInt64 = static_cast<int64>(VariantData.Get<bool>());
-	}
-	else if (VariantData.IsType<double>())
+	// Populate known schema.
+	// Validation will be handled later due to the dependency that schema has on service descriptors.
+	TMap<FSchemaId, const FSchemaDescriptor*> KnownSchemaDescriptors;
+	for (const FSchemaDescriptor& SchemaDescriptor : Config.SchemaDescriptors)
 	{
-		AsInt64 = static_cast<int64>(VariantData.Get<double>());
+		// Check that schema descriptor id has not already been used.
+		if (KnownSchemaDescriptors.Find(SchemaDescriptor.Id))
+		{
+			UE_LOG(LogOnlineSchema, Error, TEXT("Duplicate schema descriptor found: %s"), *SchemaDescriptor.Id.ToString().ToLower());
+			ParsedSuccessfully = false;
+			continue;
+		}
+
+		KnownSchemaDescriptors.Add(SchemaDescriptor.Id, &SchemaDescriptor);
 	}
-	else if (VariantData.IsType<FString>())
+
+	// Verify and populate known service descriptors.
+	TMap<FSchemaId, const FSchemaServiceDescriptor*> KnownSchemaServiceDescriptors;
+	for (const FSchemaServiceDescriptor& SchemaServiceDescriptor : Config.ServiceDescriptors)
 	{
-		::LexFromString(AsInt64, *VariantData.Get<FString>());
+		// Check that schema service descriptor id has not already been used.
+		if (KnownSchemaServiceDescriptors.Find(SchemaServiceDescriptor.Id))
+		{
+			UE_LOG(LogOnlineSchema, Error, TEXT("Duplicate schema service descriptor found: %s"), *SchemaServiceDescriptor.Id.ToString().ToLower());
+			ParsedSuccessfully = false;
+			continue;
+		}
+
+		// Verify attributes
+		TSet<FSchemaServiceAttributeId> SeenAttributes;
+		for (const FSchemaServiceAttributeDescriptor& SchemaServiceAttributeDescriptor : SchemaServiceDescriptor.Attributes)
+		{
+			// Check that attribute ids are not reused.
+			if (SeenAttributes.Find(SchemaServiceAttributeDescriptor.Id))
+			{
+				UE_LOG(LogOnlineSchema, Error, TEXT("Invalid schema service attribute %s.%s: Attribute id has already been used."),
+					*SchemaServiceDescriptor.Id.ToString().ToLower(),
+					*SchemaServiceAttributeDescriptor.Id.ToString().ToLower());
+				ParsedSuccessfully = false;
+				continue;
+			}
+			SeenAttributes.Add(SchemaServiceAttributeDescriptor.Id);
+
+			ESchemaServiceAttributeSupportedTypeFlags SupportedTypes = BuildFlagsFromArray(SchemaServiceAttributeDescriptor.SupportedTypes);
+
+			// Check that supported type has been set.
+			if (SupportedTypes == ESchemaServiceAttributeSupportedTypeFlags::None)
+			{
+				UE_LOG(LogOnlineSchema, Error, TEXT("Invalid schema service attribute %s.%s: A valid supported type must be selected."),
+					*SchemaServiceDescriptor.Id.ToString().ToLower(),
+					*SchemaServiceAttributeDescriptor.Id.ToString().ToLower());
+				ParsedSuccessfully = false;
+			}
+
+			// Check that variable length data has a set maximum length.
+			if (EnumHasAnyFlags(SupportedTypes, ESchemaServiceAttributeSupportedTypeFlags::String) && SchemaServiceAttributeDescriptor.MaxSize <= 0)
+			{
+				UE_LOG(LogOnlineSchema, Error, TEXT("Invalid schema service attribute %s.%s: A valid max size must be set for variable length data."),
+					*SchemaServiceDescriptor.Id.ToString().ToLower(),
+					*SchemaServiceAttributeDescriptor.Id.ToString().ToLower());
+				ParsedSuccessfully = false;
+			}
+		}
+
+		KnownSchemaServiceDescriptors.Add(SchemaServiceDescriptor.Id, &SchemaServiceDescriptor);
 	}
-	return AsInt64;
-}
 
-double FSchemaVariant::GetDouble() const
-{
-	double AsDouble = 0;
-	if (VariantData.IsType<double>())
+	// Verify schema data.
+	for (const TPair<FSchemaId, const FSchemaDescriptor*>& SchemaDescriptorPair : KnownSchemaDescriptors)
 	{
-		AsDouble = VariantData.Get<double>();
+		const FSchemaDescriptor& SchemaDescriptor = *SchemaDescriptorPair.Get<1>();
+
+		// Verify parent schema.
+		if (SchemaDescriptor.ParentId.IsValid())
+		{
+			TSet<FSchemaId> SeenSchema;
+			const FSchemaDescriptor* BaseParentSchemaDescriptor = nullptr;
+			const FSchemaDescriptor* TestSchemaDescriptor = &SchemaDescriptor;
+			while (TestSchemaDescriptor)
+			{
+				// Check for circular dependencies in schema.
+				if (SeenSchema.Find(TestSchemaDescriptor->Id))
+				{
+					UE_LOG(LogOnlineSchema, Error, TEXT("Circular parent dependency found in schema: %s"), *SchemaDescriptor.Id.ToString().ToLower());
+					ParsedSuccessfully = false;
+					break;
+				}
+				SeenSchema.Add(TestSchemaDescriptor->Id);
+
+				// Iterate parent.
+				if (TestSchemaDescriptor->ParentId != FSchemaId())
+				{
+					// Check that parent schema exists.
+					if (const FSchemaDescriptor** FoundDescriptor = KnownSchemaDescriptors.Find(TestSchemaDescriptor->ParentId))
+					{
+						TestSchemaDescriptor = *FoundDescriptor;
+						BaseParentSchemaDescriptor = TestSchemaDescriptor;
+					}
+					else
+					{
+						UE_LOG(LogOnlineSchema, Error, TEXT("Invalid schema %s: Parent schema %s does not exist."),
+							*SchemaDescriptor.Id.ToString().ToLower(),
+							*TestSchemaDescriptor->ParentId.ToString().ToLower());
+						ParsedSuccessfully = false;
+						TestSchemaDescriptor = nullptr;
+					}
+				}
+				else
+				{
+					TestSchemaDescriptor = nullptr;
+				}
+			}
+
+			// verify that all categories used in schema exist in base parent schema.
+			if (BaseParentSchemaDescriptor)
+			{
+				for (const FSchemaCategoryDescriptor& SchemaCategoryDescriptor : SchemaDescriptor.Categories)
+				{
+					bool bBaseCategoryExists = Algo::FindByPredicate(BaseParentSchemaDescriptor->Categories,
+					[&SchemaCategoryDescriptor](const FSchemaCategoryDescriptor& Descriptor)
+					{
+						return Descriptor.Id == SchemaCategoryDescriptor.Id;
+					}) != nullptr;
+
+					if (!bBaseCategoryExists)
+					{
+						UE_LOG(LogOnlineSchema, Error, TEXT("Invalid schema %s: Category %s does not exist in base parent %s schema."),
+							*SchemaDescriptor.Id.ToString().ToLower(),
+							*SchemaCategoryDescriptor.Id.ToString().ToLower(),
+							*BaseParentSchemaDescriptor->Id.ToString().ToLower());
+							ParsedSuccessfully = false;
+					}
+				}
+			}
+		}
+
+		// Verify category descriptors.
+		TSet<FSchemaCategoryId> SeenSchemaCategoryIds;
+		TSet<FSchemaServiceDescriptorId> SeenSchemaServiceDescriptorIds;
+		for (const FSchemaCategoryDescriptor& SchemaCategoryDescriptor : SchemaDescriptor.Categories)
+		{
+			if (SeenSchemaCategoryIds.Find(SchemaCategoryDescriptor.Id))
+			{
+				UE_LOG(LogOnlineSchema, Error, TEXT("Invalid schema category %s.%s: Category id has already been used."),
+					*SchemaDescriptor.Id.ToString().ToLower(),
+					*SchemaCategoryDescriptor.Id.ToString().ToLower());
+				ParsedSuccessfully = false;
+				continue;
+			}
+			SeenSchemaCategoryIds.Add(SchemaCategoryDescriptor.Id);
+
+			// Verify that service descriptor is valid.
+			if (SchemaCategoryDescriptor.ServiceDescriptorId != FSchemaServiceDescriptorId())
+			{
+				// Verify that only a base parent set service descriptor.
+				if (SchemaDescriptor.ParentId != FSchemaId())
+				{
+					UE_LOG(LogOnlineSchema, Error, TEXT("Invalid schema category %s.%s: Service descriptor id may only be set by a base parent schmea category."),
+						*SchemaDescriptor.Id.ToString().ToLower(),
+						*SchemaCategoryDescriptor.Id.ToString().ToLower(),
+						*SchemaCategoryDescriptor.ServiceDescriptorId.ToString().ToLower());
+					ParsedSuccessfully = false;
+				}
+
+				// Verify that service descriptor exists.
+				if (KnownSchemaServiceDescriptors.Find(SchemaCategoryDescriptor.ServiceDescriptorId) == nullptr)
+				{
+					UE_LOG(LogOnlineSchema, Error, TEXT("Invalid schema category %s.%s: Service descriptor id %s not found."),
+						*SchemaDescriptor.Id.ToString().ToLower(),
+						*SchemaCategoryDescriptor.Id.ToString().ToLower(),
+						*SchemaCategoryDescriptor.ServiceDescriptorId.ToString().ToLower());
+					ParsedSuccessfully = false;
+				}
+
+				// Verify that the service descriptor is not reused between categories.
+				if (SeenSchemaServiceDescriptorIds.Find(SchemaCategoryDescriptor.ServiceDescriptorId) != nullptr)
+				{
+					UE_LOG(LogOnlineSchema, Error, TEXT("Invalid schema category %s.%s: Service descriptor id %s has already been used by another category."),
+						*SchemaDescriptor.Id.ToString().ToLower(),
+						*SchemaCategoryDescriptor.Id.ToString().ToLower(),
+						*SchemaCategoryDescriptor.ServiceDescriptorId.ToString().ToLower());
+					ParsedSuccessfully = false;
+				}
+
+				SeenSchemaServiceDescriptorIds.Add(SchemaCategoryDescriptor.ServiceDescriptorId);
+			}
+
+			// Verify category attributes.
+			TSet<FSchemaServiceAttributeId> SeenAttributes;
+			FSchemaAttributeId SeenSchemaCompatibilityAttributeId;
+			for (const FSchemaAttributeDescriptor& SchemaAttributeDescriptor : SchemaCategoryDescriptor.Attributes)
+			{
+				// Check that attribute ids are not reused.
+				if (SeenAttributes.Find(SchemaAttributeDescriptor.Id))
+				{
+					UE_LOG(LogOnlineSchema, Error, TEXT("Invalid schema attribute %s.%s.%s: Attribute id has already been used."),
+						*SchemaDescriptor.Id.ToString().ToLower(),
+						*SchemaCategoryDescriptor.Id.ToString().ToLower(),
+						*SchemaAttributeDescriptor.Id.ToString().ToLower());
+					ParsedSuccessfully = false;
+					continue;
+				}
+				SeenAttributes.Add(SchemaAttributeDescriptor.Id);
+
+				// Check that variable length data has a set maximum length.
+				if (SchemaAttributeDescriptor.Type == ESchemaAttributeType::String && SchemaAttributeDescriptor.MaxSize <= 0)
+				{
+					UE_LOG(LogOnlineSchema, Error, TEXT("Invalid schema attribute %s.%s.%s: A valid max size must be set for variable length data."),
+						*SchemaDescriptor.Id.ToString().ToLower(),
+						*SchemaCategoryDescriptor.Id.ToString().ToLower(),
+						*SchemaAttributeDescriptor.Id.ToString().ToLower());
+					ParsedSuccessfully = false;
+				}
+
+				ESchemaAttributeFlags AttributeFlags = BuildFlagsFromArray(SchemaAttributeDescriptor.Flags);
+
+				// Check that searchable attributes only exist in a base schema.
+				if (EnumHasAllFlags(AttributeFlags, ESchemaAttributeFlags::Searchable) && SchemaDescriptor.ParentId != FSchemaId())
+				{
+					UE_LOG(LogOnlineSchema, Error, TEXT("Invalid schema attribute %s.%s.%s: Searchable fields may only exist in the base schema."),
+						*SchemaDescriptor.Id.ToString().ToLower(),
+						*SchemaCategoryDescriptor.Id.ToString().ToLower(),
+						*SchemaAttributeDescriptor.Id.ToString().ToLower());
+					ParsedSuccessfully = false;
+				}
+
+				// Check that schema compatibility attribute only exists in a base schema and used only once and that it is the Int64 type.
+				if (EnumHasAllFlags(AttributeFlags, ESchemaAttributeFlags::SchemaCompatibilityId))
+				{
+					if (SchemaDescriptor.ParentId != FSchemaId())
+					{
+						UE_LOG(LogOnlineSchema, Error, TEXT("Invalid schema attribute %s.%s.%s: SchemaCompatibilityId field may only exist in the base schema."),
+							*SchemaDescriptor.Id.ToString().ToLower(),
+							*SchemaCategoryDescriptor.Id.ToString().ToLower(),
+							*SchemaAttributeDescriptor.Id.ToString().ToLower());
+						ParsedSuccessfully = false;
+					}
+
+					if (SeenSchemaCompatibilityAttributeId != FSchemaId())
+					{
+						UE_LOG(LogOnlineSchema, Error, TEXT("Invalid schema attribute %s.%s.%s: SchemaCompatibilityId field has already been used by attribute %s."),
+							*SchemaDescriptor.Id.ToString().ToLower(),
+							*SchemaCategoryDescriptor.Id.ToString().ToLower(),
+							*SchemaAttributeDescriptor.Id.ToString().ToLower(),
+							*SeenSchemaCompatibilityAttributeId.ToString().ToLower());
+						ParsedSuccessfully = false;
+					}
+					else
+					{
+						SeenSchemaCompatibilityAttributeId = SchemaAttributeDescriptor.Id;
+					}
+
+					if (SchemaAttributeDescriptor.Type != ESchemaAttributeType::Int64)
+					{
+						UE_LOG(LogOnlineSchema, Error, TEXT("Invalid schema attribute %s.%s.%s: SchemaCompatibilityId field may only be set to Int64."),
+							*SchemaDescriptor.Id.ToString().ToLower(),
+							*SchemaCategoryDescriptor.Id.ToString().ToLower(),
+							*SchemaAttributeDescriptor.Id.ToString().ToLower());
+						ParsedSuccessfully = false;
+					}
+				}
+
+				// Check that group id is valid.
+				if (SchemaAttributeDescriptor.UpdateGroupId < 0)
+				{
+					UE_LOG(LogOnlineSchema, Error, TEXT("Invalid schema attribute %s.%s.%s: UpdateGroupId must be set to either 0 or a positive integer value."),
+						*SchemaDescriptor.Id.ToString().ToLower(),
+						*SchemaCategoryDescriptor.Id.ToString().ToLower(),
+						*SchemaAttributeDescriptor.Id.ToString().ToLower());
+					ParsedSuccessfully = false;
+				}
+			}
+		}
 	}
-	else if (VariantData.IsType<FString>())
+
+	// Build definitions only when there are no parsing errors.
+	TMap<FSchemaId, TSharedRef<const FSchemaDefinition>> ParsedSchemaDefinitions;
+	if (ParsedSuccessfully)
 	{
-		::LexFromString(AsDouble, *VariantData.Get<FString>());
+		// Build schema definitions
+		TSet<uint32> SeenSchemaCrcs;
+		for (const TPair<FSchemaId, const FSchemaDescriptor*>& SchemaDescriptorPair : KnownSchemaDescriptors)
+		{
+			uint32 SchemaDataCrc = 0;
+			const FSchemaDescriptor& SchemaDescriptor = *SchemaDescriptorPair.Get<1>();
+
+			TSharedRef<FSchemaDefinition> SchemaDefinition = MakeShared<FSchemaDefinition>();
+			ParsedSchemaDefinitions.Add(SchemaDescriptor.Id, SchemaDefinition);
+			SchemaDefinition->Id = SchemaDescriptor.Id;
+
+			TMap<FSchemaCategoryId, FSchemaServiceDescriptorId> ServiceDescriptors;
+
+			// Populate definition from descriptor
+			for (const FSchemaDescriptor** IterSchemaDescriptor = KnownSchemaDescriptors.Find(SchemaDescriptor.Id);
+				IterSchemaDescriptor != nullptr;
+				IterSchemaDescriptor = KnownSchemaDescriptors.Find((*IterSchemaDescriptor)->ParentId))
+			{
+				const FSchemaDescriptor& CurrentSchemaDescriptor = **IterSchemaDescriptor;
+
+				// Populate parent info.
+				if (CurrentSchemaDescriptor.Id != SchemaDescriptor.Id)
+				{
+					SchemaDefinition->ParentSchemaIds.Add(CurrentSchemaDescriptor.Id);
+					SchemaDataCrc = FCrc::TypeCrc32(CurrentSchemaDescriptor.Id, SchemaDataCrc);
+				}
+
+				// Populate attribute definitions.
+				for (const FSchemaCategoryDescriptor& CurrentSchemaCategoryDescriptor : CurrentSchemaDescriptor.Categories)
+				{
+					FSchemaCategoryDefinition& SchemaCategoryDefinition = SchemaDefinition->Categories.FindOrAdd(CurrentSchemaCategoryDescriptor.Id);
+					SchemaDataCrc = FCrc::TypeCrc32(CurrentSchemaCategoryDescriptor.Id, SchemaDataCrc);
+
+					for (const FSchemaAttributeDescriptor& CurrentSchemaAttributeDescriptor : CurrentSchemaCategoryDescriptor.Attributes)
+					{
+						FSchemaAttributeDefinition& AttributeDefinition = SchemaCategoryDefinition.SchemaAttributeDefinitions.Add(CurrentSchemaAttributeDescriptor.Id);
+						AttributeDefinition.Id = CurrentSchemaAttributeDescriptor.Id;
+						AttributeDefinition.Flags = BuildFlagsFromArray(CurrentSchemaAttributeDescriptor.Flags);
+						AttributeDefinition.Type = CurrentSchemaAttributeDescriptor.Type;
+						AttributeDefinition.MaxSize = CurrentSchemaAttributeDescriptor.MaxSize;
+						SchemaDataCrc = FCrc::TypeCrc32(AttributeDefinition.Id, SchemaDataCrc);
+						SchemaDataCrc = FCrc::TypeCrc32(AttributeDefinition.Flags, SchemaDataCrc);
+						SchemaDataCrc = FCrc::TypeCrc32(AttributeDefinition.Type, SchemaDataCrc);
+						SchemaDataCrc = FCrc::TypeCrc32(AttributeDefinition.MaxSize, SchemaDataCrc);
+					}
+
+					// Store service descriptor for assignment.
+					if (CurrentSchemaCategoryDescriptor.ServiceDescriptorId != FSchemaServiceDescriptorId())
+					{
+						ServiceDescriptors.Add(CurrentSchemaCategoryDescriptor.Id, CurrentSchemaCategoryDescriptor.ServiceDescriptorId);
+					}
+				}
+			}
+
+			// Assign attributes to service attributes.
+			for (TPair<FSchemaCategoryId, FSchemaServiceDescriptorId>& ServiceDescriptor : ServiceDescriptors)
+			{
+				FSchemaCategoryDefinition& SchemaCategoryDefinition = *SchemaDefinition->Categories.Find(ServiceDescriptor.Get<0>());
+				const FSchemaServiceDescriptor& SchemaServiceDescriptor = **KnownSchemaServiceDescriptors.Find(ServiceDescriptor.Get<1>());
+
+				// todo: pack multiple attributes into a service attribute.
+				// For the time being schema attributes and service attributes have a 1:1 relationship.
+
+				// todo
+			}
+
+			// Build compatibility id.
+			{
+				const uint32 SchemaIdCrc = FCrc::TypeCrc32(SchemaDefinition->Id);
+				if (SeenSchemaCrcs.Find(SchemaIdCrc) != nullptr)
+				{
+					UE_LOG(LogOnlineSchema, Error, TEXT("Invalid schema %s: Hash collision processing schema name, please rename."),
+						*SchemaDefinition->Id.ToString().ToLower());
+					ParsedSuccessfully = false;
+				}
+				SeenSchemaCrcs.Add(SchemaIdCrc);
+
+				const uint64 CompatibilityIdUint = (static_cast<uint64>(SchemaIdCrc) << 32) | SchemaDataCrc;
+				SchemaDefinition->CompatibilityId = *reinterpret_cast<const int64*>(&CompatibilityIdUint);
+			}
+		}
 	}
-	else if (VariantData.IsType<int64>())
+
+	if (ParsedSuccessfully)
 	{
-		AsDouble = static_cast<double>(VariantData.Get<int64>());
+		SchemaDefinitions = MoveTemp(ParsedSchemaDefinitions);
+		return true;
 	}
-	else if (VariantData.IsType<bool>())
+	else
 	{
-		AsDouble = static_cast<double>(VariantData.Get<bool>());
+		return false;
 	}
-	return AsDouble;
 }
 
-bool FSchemaVariant::GetBoolean() const
+FSchemaCategoryInstance::FSchemaCategoryInstance(
+	const FSchemaId& SchemaId,
+	const FSchemaId& BaseSchemaId,
+	const FSchemaCategoryId& CategoryId,
+	const TSharedRef<const FSchemaRegistry>& Registry)
 {
-	bool bAsBool = false;
-	if (VariantData.IsType<bool>())
+	if ((SchemaDefinition = Registry->GetDefinition(SchemaId)))
 	{
-		bAsBool = VariantData.Get<bool>();
-	}
-	else if (VariantData.IsType<FString>())
-	{
-		::LexFromString(bAsBool, *VariantData.Get<FString>());
+		if (!Registry->IsSchemaChildOf(SchemaId, BaseSchemaId))
+		{
+			UE_LOG(LogOnlineSchema, Error, TEXT("SchemaCategoryInstance init error: Schema %s is not a child of %s"),
+				*SchemaId.ToString().ToLower(),
+				*BaseSchemaId.ToString().ToLower());
+		}
+		else
+		{
+			SchemaCategoryDefinition = SchemaDefinition->Categories.Find(CategoryId);
+			if (SchemaCategoryDefinition == nullptr)
+			{
+				UE_LOG(LogOnlineSchema, Error, TEXT("SchemaCategoryInstance init error: Unable to find category %s in schema %s"),
+					*SchemaId.ToString().ToLower());
+			}
+		}
 	}
 	else
 	{
-		bAsBool = GetInt64() != 0;
+		UE_LOG(LogOnlineSchema, Error, TEXT("SchemaCategoryInstance init error: Unable to find definition for schema %s"),
+			*SchemaId.ToString().ToLower());
 	}
-	return bAsBool;
 }
 
-FString FSchemaVariant::GetString() const
+TOnlineResult<FSchemaCategoryInstanceApplyApplicationChanges> FSchemaCategoryInstance::ApplyChanges(FSchemaCategoryInstanceApplyApplicationChanges::Params&& Changes)
 {
-	if (VariantData.IsType<FString>())
-	{
-		return VariantData.Get<FString>();
-	}
-	else if (VariantData.IsType<int64>())
-	{
-		return FString::Printf(TEXT("%" INT64_FMT), VariantData.Get<int64>());
-	}
-	else if (VariantData.IsType<bool>())
-	{
-		return ::LexToString(VariantData.Get<bool>());
-	}
-	else if (VariantData.IsType<double>())
+	// todo: handle packing multiple attributes into a service attribute.
+
+	// todo
+	return TOnlineResult<FSchemaCategoryInstanceApplyApplicationChanges>(FSchemaCategoryInstanceApplyApplicationChanges::Result({}));
+}
+
+TOnlineResult<FSchemaCategoryInstanceApplyServiceChanges> FSchemaCategoryInstance::ApplyChanges(FSchemaCategoryInstanceApplyServiceChanges::Params&& Changes)
+{
+	// todo: handle unpacking a service attribute into multiple attributes.
+
+	// todo
+	return TOnlineResult<FSchemaCategoryInstanceApplyServiceChanges>(FSchemaCategoryInstanceApplyServiceChanges::Result({}));
+}
+
+bool FSchemaCategoryInstance::IsValid() const
+{
+	return SchemaCategoryDefinition != nullptr;
+}
+
+bool FSchemaCategoryInstance::IsValidAttributeData(const FSchemaAttributeId& Id, const FSchemaVariant& Data)
+{
+	if (ensure(IsValid()))
 	{
-		return FString::Printf(TEXT("%f"), VariantData.Get<double>());
+		if (const FSchemaAttributeDefinition* AttributeDefinition = SchemaCategoryDefinition->SchemaAttributeDefinitions.Find(Id))
+		{
+			if (AttributeDefinition->Type == Data.GetType())
+			{
+				return true;
+			}
+			else
+			{
+				UE_LOG(LogOnlineSchema, Verbose, TEXT("[IsValidAttributeData] Schema attribute %s.%s.%s set with invalid type %s. Expected type %s"),
+					*SchemaDefinition->Id.ToString().ToLower(),
+					*SchemaCategoryDefinition->Id.ToString().ToLower(),
+					*AttributeDefinition->Id.ToString().ToLower(),
+					*LexToString(Data.GetType()),
+					LexToString(AttributeDefinition->Type));
+				return false;
+			}
+		}
+		else
+		{
+			UE_LOG(LogOnlineSchema, Verbose, TEXT("[IsValidAttributeData] Attribute %s not found in schema definition %s.%s"),
+				*Id.ToString().ToLower(),
+				*SchemaDefinition->Id.ToString().ToLower(),
+				*SchemaCategoryDefinition->Id.ToString().ToLower());
+			return false;
+		}
 	}
 	else
 	{
-		checkNoEntry();
+		UE_LOG(LogOnlineSchema, Verbose, TEXT("[IsValidAttributeData] Unable to set attribute %s. Schema is not valid"),
+			*Id.ToString().ToLower());
+		return false;
 	}
-	return TEXT("");
 }
 
-bool FSchemaVariant::operator==(const FSchemaVariant& Other) const
+TSharedPtr<const FSchemaDefinition> FSchemaRegistry::GetDefinition(const FSchemaId& SchemaId) const
 {
-	if (VariantData.GetIndex() != Other.VariantData.GetIndex())
-	{
-		return false;
-	}
+	const TSharedRef<const FSchemaDefinition>* Definition = SchemaDefinitions.Find(SchemaId);
+	return Definition ? *Definition : TSharedPtr<const FSchemaDefinition>();
+}
 
-	switch (VariantData.GetIndex())
+bool FSchemaRegistry::IsSchemaChildOf(const FSchemaId& SchemaId, const FSchemaId& ParentSchemaId) const
+{
+	const FSchemaId* FoundSchemaId = nullptr;
+	if (const TSharedRef<const FSchemaDefinition>* Definition = SchemaDefinitions.Find(SchemaId))
 	{
-	case FVariantType::IndexOfType<FString>():	return VariantData.Get<FString>() == Other.VariantData.Get<FString>();
-	case FVariantType::IndexOfType<int64>():	return VariantData.Get<int64>() == Other.VariantData.Get<int64>();
-	case FVariantType::IndexOfType<double>():	return VariantData.Get<double>() == Other.VariantData.Get<double>();
-
-	default:checkNoEntry(); // Intentional fallthrough
-	case FVariantType::IndexOfType<bool>():	return VariantData.Get<bool>() == Other.VariantData.Get<bool>();
+		FoundSchemaId = (* Definition)->ParentSchemaIds.Find(ParentSchemaId);
 	}
+	return FoundSchemaId != nullptr;
 }
 
 /* UE::Online */ }
diff --git a/Engine/Plugins/Online/Experimental/OnlineServices/Source/OnlineServicesInterface/Private/Online/SchemaTypes.cpp b/Engine/Plugins/Online/Experimental/OnlineServices/Source/OnlineServicesInterface/Private/Online/SchemaTypes.cpp
new file mode 100644
--- /dev/null
+++ b/Engine/Plugins/Online/Experimental/OnlineServices/Source/OnlineServicesInterface/Private/Online/SchemaTypes.cpp
@@ -0,0 +1,304 @@
+// Copyright Epic Games, Inc. All Rights Reserved.
+
+#include "Online/SchemaTypes.h"
+
+namespace UE::Online {
+
+const TCHAR* LexToString(ESchemaAttributeFlags SchemaAttributeFlags)
+{
+	switch (SchemaAttributeFlags)
+	{
+	default:								checkNoEntry(); // Intentional fallthrough
+	case ESchemaAttributeFlags::None:		return TEXT("None");
+	case ESchemaAttributeFlags::Searchable:	return TEXT("Searchable");
+	case ESchemaAttributeFlags::Public:		return TEXT("Public");
+	}
+}
+
+void LexFromString(ESchemaAttributeFlags& OutSchemaAttributeFlags, const TCHAR* InStr)
+{
+	if (FCString::Stricmp(InStr, TEXT("None")) == 0)
+	{
+		OutSchemaAttributeFlags = ESchemaAttributeFlags::None;
+	}
+	else if (FCString::Stricmp(InStr, TEXT("Searchable")) == 0)
+	{
+		OutSchemaAttributeFlags = ESchemaAttributeFlags::Searchable;
+	}
+	else if (FCString::Stricmp(InStr, TEXT("Public")) == 0)
+	{
+		OutSchemaAttributeFlags = ESchemaAttributeFlags::Public;
+	}
+	else if (FCString::Stricmp(InStr, TEXT("SchemaCompatibilityId")) == 0)
+	{
+		OutSchemaAttributeFlags = ESchemaAttributeFlags::SchemaCompatibilityId;
+	}
+	else
+	{
+		checkNoEntry();
+		OutSchemaAttributeFlags = ESchemaAttributeFlags::None;
+	}
+}
+
+const TCHAR* LexToString(ESchemaAttributeType SchemaAttributeType)
+{
+	switch (SchemaAttributeType)
+	{
+	default:							checkNoEntry(); // Intentional fallthrough
+	case ESchemaAttributeType::Bool:	return TEXT("Bool");
+	case ESchemaAttributeType::Int64:	return TEXT("Int64");
+	case ESchemaAttributeType::Double:	return TEXT("Double");
+	case ESchemaAttributeType::String:	return TEXT("String");
+	}
+}
+
+void LexFromString(ESchemaAttributeType& OutSchemaAttributeType, const TCHAR* InStr)
+{
+	if (FCString::Stricmp(InStr, TEXT("Bool")) == 0)
+	{
+		OutSchemaAttributeType = ESchemaAttributeType::Bool;
+	}
+	else if (FCString::Stricmp(InStr, TEXT("Int64")) == 0)
+	{
+		OutSchemaAttributeType = ESchemaAttributeType::Int64;
+	}
+	else if (FCString::Stricmp(InStr, TEXT("Double")) == 0)
+	{
+		OutSchemaAttributeType = ESchemaAttributeType::Double;
+	}
+	else if (FCString::Stricmp(InStr, TEXT("String")) == 0)
+	{
+		OutSchemaAttributeType = ESchemaAttributeType::String;
+	}
+
+	else
+	{
+		checkNoEntry();
+		OutSchemaAttributeType = ESchemaAttributeType::Bool;
+	}
+}
+
+const TCHAR* LexToString(ESchemaServiceAttributeFlags SchemaAttributeFlags)
+{
+	switch (SchemaAttributeFlags)
+	{
+	default:										checkNoEntry(); // Intentional fallthrough
+	case ESchemaServiceAttributeFlags::None:		return TEXT("None");
+	case ESchemaServiceAttributeFlags::Searchable:	return TEXT("Searchable");
+	case ESchemaServiceAttributeFlags::Public:		return TEXT("Public");
+	}
+}
+
+void LexFromString(ESchemaServiceAttributeFlags& OutSchemaAttributeFlags, const TCHAR* InStr)
+{
+	if (FCString::Stricmp(InStr, TEXT("None")) == 0)
+	{
+		OutSchemaAttributeFlags = ESchemaServiceAttributeFlags::None;
+	}
+	else if (FCString::Stricmp(InStr, TEXT("Searchable")) == 0)
+	{
+		OutSchemaAttributeFlags = ESchemaServiceAttributeFlags::Searchable;
+	}
+	else if (FCString::Stricmp(InStr, TEXT("Public")) == 0)
+	{
+		OutSchemaAttributeFlags = ESchemaServiceAttributeFlags::Public;
+	}
+	else
+	{
+		checkNoEntry();
+		OutSchemaAttributeFlags = ESchemaServiceAttributeFlags::None;
+	}
+}
+
+const TCHAR* LexToString(ESchemaServiceAttributeSupportedTypeFlags SchemaServiceAttributeTypeFlags)
+{
+	switch (SchemaServiceAttributeTypeFlags)
+	{
+	default:												checkNoEntry(); // Intentional fallthrough
+	case ESchemaServiceAttributeSupportedTypeFlags::Bool:	return TEXT("Bool");
+	case ESchemaServiceAttributeSupportedTypeFlags::Int64:	return TEXT("Int64");
+	case ESchemaServiceAttributeSupportedTypeFlags::Double:	return TEXT("Double");
+	case ESchemaServiceAttributeSupportedTypeFlags::String:	return TEXT("String");
+	}
+}
+
+void LexFromString(ESchemaServiceAttributeSupportedTypeFlags& OutSchemaServiceAttributeTypeFlags, const TCHAR* InStr)
+{
+	if (FCString::Stricmp(InStr, TEXT("Bool")) == 0)
+	{
+		OutSchemaServiceAttributeTypeFlags = ESchemaServiceAttributeSupportedTypeFlags::Bool;
+	}
+	else if (FCString::Stricmp(InStr, TEXT("Int64")) == 0)
+	{
+		OutSchemaServiceAttributeTypeFlags = ESchemaServiceAttributeSupportedTypeFlags::Int64;
+	}
+	else if (FCString::Stricmp(InStr, TEXT("Double")) == 0)
+	{
+		OutSchemaServiceAttributeTypeFlags = ESchemaServiceAttributeSupportedTypeFlags::Double;
+	}
+	else if (FCString::Stricmp(InStr, TEXT("String")) == 0)
+	{
+		OutSchemaServiceAttributeTypeFlags = ESchemaServiceAttributeSupportedTypeFlags::String;
+	}
+	else
+	{
+		checkNoEntry();
+		OutSchemaServiceAttributeTypeFlags = ESchemaServiceAttributeSupportedTypeFlags::Bool;
+	}
+}
+
+FSchemaVariant::FSchemaVariant(FSchemaVariant&& InOther)
+{
+	VariantData = MoveTemp(InOther.VariantData);
+	VariantType = InOther.VariantType;
+	InOther.VariantType = ESchemaAttributeType::None;
+}
+
+FSchemaVariant& FSchemaVariant::operator=(FSchemaVariant&& InOther)
+{
+	VariantData = MoveTemp(InOther.VariantData);
+	VariantType = InOther.VariantType;
+	InOther.VariantType = ESchemaAttributeType::None;
+	return *this;
+}
+
+void FSchemaVariant::Set(const TCHAR* AsString)
+{
+	VariantData.Emplace<FString>(AsString);
+	VariantType = ESchemaAttributeType::String;
+}
+
+void FSchemaVariant::Set(const FString& AsString)
+{
+	VariantData.Emplace<FString>(AsString);
+	VariantType = ESchemaAttributeType::String;
+}
+
+void FSchemaVariant::Set(FString&& AsString)
+{
+	VariantData.Emplace<FString>(MoveTemp(AsString));
+	VariantType = ESchemaAttributeType::String;
+}
+
+void FSchemaVariant::Set(int64 AsInt)
+{
+	VariantData.Emplace<int64>(AsInt);
+	VariantType = ESchemaAttributeType::Int64;
+}
+
+void FSchemaVariant::Set(double AsDouble)
+{
+	VariantData.Emplace<double>(AsDouble);
+	VariantType = ESchemaAttributeType::Double;
+}
+
+void FSchemaVariant::Set(bool bAsBool)
+{
+	VariantData.Emplace<bool>(bAsBool);
+	VariantType = ESchemaAttributeType::Bool;
+}
+
+int64 FSchemaVariant::GetInt64() const
+{
+	if (ensure(VariantType == ESchemaAttributeType::Int64))
+	{
+		return VariantData.Get<int64>();
+	}
+	else
+	{
+		return 0;
+	}
+}
+
+double FSchemaVariant::GetDouble() const
+{
+	if (ensure(VariantType == ESchemaAttributeType::Double))
+	{
+		return VariantData.Get<double>();
+	}
+	else
+	{
+		return 0;
+	}
+}
+
+bool FSchemaVariant::GetBoolean() const
+{
+	if (ensure(VariantType == ESchemaAttributeType::Bool))
+	{
+		return VariantData.Get<bool>();
+	}
+	else
+	{
+		return 0;
+	}
+}
+
+FString FSchemaVariant::GetString() const
+{
+	if (ensure(VariantType == ESchemaAttributeType::String))
+	{
+		return VariantData.Get<FString>();
+	}
+	else
+	{
+		return FString();
+	}
+}
+
+FString FSchemaVariant::ToLogString() const
+{
+	FString ValueString;
+
+	if (VariantData.IsType<FString>())
+	{
+		ValueString = VariantData.Get<FString>();
+	}
+	else if (VariantData.IsType<int64>())
+	{
+		ValueString = FString::Printf(TEXT("%" INT64_FMT), VariantData.Get<int64>());
+	}
+	else if (VariantData.IsType<bool>())
+	{
+		ValueString = ::LexToString(VariantData.Get<bool>());
+	}
+	else if (VariantData.IsType<double>())
+	{
+		ValueString = FString::Printf(TEXT("%.3f"), VariantData.Get<double>());
+	}
+
+	return FString::Printf(TEXT("%s:%s"), LexToString(VariantType), *ValueString);
+}
+
+bool FSchemaVariant::operator==(const FSchemaVariant& Other) const
+{
+	if (VariantData.GetIndex() != Other.VariantData.GetIndex())
+	{
+		return false;
+	}
+
+	// Type should always match due to the variant data.
+	check(VariantType == Other.VariantType);
+
+	switch (VariantData.GetIndex())
+	{
+	case FVariantType::IndexOfType<FString>():	return VariantData.Get<FString>() == Other.VariantData.Get<FString>();
+	case FVariantType::IndexOfType<int64>():	return VariantData.Get<int64>() == Other.VariantData.Get<int64>();
+	case FVariantType::IndexOfType<double>():	return VariantData.Get<double>() == Other.VariantData.Get<double>();
+
+	default:									checkNoEntry(); // Intentional fallthrough
+	case FVariantType::IndexOfType<bool>():		return VariantData.Get<bool>() == Other.VariantData.Get<bool>();
+	}
+}
+
+const FString LexToString(const FLexToStringAdaptor<FSchemaVariant>& Adaptor)
+{
+	return Adaptor.SchemaVariant.ToLogString();
+}
+
+void LexFromString(FSchemaVariant& OutSchemaVariant, const TCHAR* InStr)
+{
+	// todo
+}
+
+/* UE::Online */ }
diff --git a/Engine/Plugins/Online/Experimental/OnlineServices/Source/OnlineServicesInterface/Public/Online/Schema.h b/Engine/Plugins/Online/Experimental/OnlineServices/Source/OnlineServicesInterface/Public/Online/Schema.h
--- a/Engine/Plugins/Online/Experimental/OnlineServices/Source/OnlineServicesInterface/Public/Online/Schema.h
+++ b/Engine/Plugins/Online/Experimental/OnlineServices/Source/OnlineServicesInterface/Public/Online/Schema.h
@@ -3,59 +3,58 @@
 #pragma once
 
 #include "CoreMinimal.h"
-#include "Misc/TVariant.h"
+#include "Online/SchemaTypes.h"
 
 namespace UE::Online {
 
-using FSchemaId = FName;
-using FSchemaAttributeId = FName;
+class FSchemaRegistry;
 
-class ONLINESERVICESINTERFACE_API FSchemaVariant final
+class ONLINESERVICESINTERFACE_API FSchemaCategoryInstance
 {
 public:
-	using FVariantType = TVariant<FString, int64, double, bool>;
-	FSchemaVariant() = default;
-	FSchemaVariant(const FSchemaVariant& InOther) = default;
-	FSchemaVariant(FSchemaVariant&& InOther) = default;
-	FSchemaVariant& operator=(FSchemaVariant&&) = default;
-	FSchemaVariant& operator=(const FSchemaVariant&) = default;
-
-	template<typename ValueType>
-	FSchemaVariant(const ValueType& InData) { Set(InData); }
-	template<typename ValueType>
-	FSchemaVariant(ValueType&& InData) { Set(MoveTemp(InData)); }
-	void Set(const TCHAR* AsString) { VariantData.Emplace<FString>(AsString); }
-	void Set(const FString& AsString) { VariantData.Emplace<FString>(AsString); }
-	void Set(FString&& AsString) { VariantData.Emplace<FString>(MoveTemp(AsString)); }
-	void Set(int64 AsInt) { VariantData.Emplace<int64>(AsInt); }
-	void Set(double AsDouble) { VariantData.Emplace<double>(AsDouble); }
-	void Set(bool bAsBool) { VariantData.Emplace<bool>(bAsBool); }
-	int64 GetInt64() const;
-	double GetDouble() const;
-	bool GetBoolean() const;
-	FString GetString() const;
-
-	bool operator==(const FSchemaVariant& Other) const;
-	bool operator!=(const FSchemaVariant& Other) const { return !(*this == Other); }
-public:
-	FVariantType VariantData;
+	FSchemaCategoryInstance(
+		const FSchemaId& SchemaId,
+		const FSchemaId& BaseSchemaId,
+		const FSchemaCategoryId& CategoryId,
+		const TSharedRef<const FSchemaRegistry>& Registry);
+
+#if 0 // todo
+	TOnlineResult<FSchemaCategoryInstanceChangeSchema> ChangeSchema(FSchemaCategoryInstanceChangeSchema::Params&& Params);
+#endif
+	TOnlineResult<FSchemaCategoryInstanceApplyApplicationChanges> ApplyChanges(FSchemaCategoryInstanceApplyApplicationChanges::Params&& Changes);
+	TOnlineResult<FSchemaCategoryInstanceApplyServiceChanges> ApplyChanges(FSchemaCategoryInstanceApplyServiceChanges::Params&& Changes);
+
+	TSharedPtr<const FSchemaDefinition> GetDefinition() const { return SchemaDefinition; };
+
+	// Check whether the category instance initialized successfully.
+	bool IsValid() const;
+
+	// Check whether an application attribute is valid in the schema category.
+	// Intended to verify attributes when there is no backing service for a schema.
+	bool IsValidAttributeData(const FSchemaAttributeId& Id, const FSchemaVariant& Data);
+
+private:
+	TSharedPtr<const FSchemaDefinition> SchemaDefinition;
+	const FSchemaCategoryDefinition* SchemaCategoryDefinition = nullptr;
+
+	// Cache of data in application form. Updates from the application may only be part of a
+	// service field, so the values need to be cached so that the service field can rebuilt when
+	// changes occur.
+	TMap<FSchemaAttributeId, FSchemaVariant> AppDataSnapshot;
 };
 
-// Don't allow implicit conversion to FSchemaVariant when calling LexToString.
-template <typename T> class FLexToStringAdaptor;
-template <>
-class FLexToStringAdaptor<FSchemaVariant>
+class ONLINESERVICESINTERFACE_API FSchemaRegistry
 {
 public:
-	FLexToStringAdaptor(const FSchemaVariant& SchemaVariant)
-		: SchemaVariant(SchemaVariant)
-	{
-	}
+	// Parse the loaded config structures.
+	bool ParseConfig(const FSchemaRegistryDescriptorConfig& Config);
 
-	const FSchemaVariant& SchemaVariant;
-};
+	TSharedPtr<const FSchemaDefinition> GetDefinition(const FSchemaId& SchemaId) const;
+	bool IsSchemaChildOf(const FSchemaId& SchemaId, const FSchemaId& ParentSchemaId) const;
 
-ONLINESERVICESINTERFACE_API const FString LexToString(const FLexToStringAdaptor<FSchemaVariant>& Adaptor);
-ONLINESERVICESINTERFACE_API void LexFromString(FSchemaVariant& Variant, const TCHAR* InStr);
+private:
+
+	TMap<FSchemaId, TSharedRef<const FSchemaDefinition>> SchemaDefinitions;
+};
 
 /* UE::Online */ }
diff --git a/Engine/Plugins/Online/Experimental/OnlineServices/Source/OnlineServicesInterface/Public/Online/SchemaTypes.h b/Engine/Plugins/Online/Experimental/OnlineServices/Source/OnlineServicesInterface/Public/Online/SchemaTypes.h
new file mode 100644
--- /dev/null
+++ b/Engine/Plugins/Online/Experimental/OnlineServices/Source/OnlineServicesInterface/Public/Online/SchemaTypes.h
@@ -0,0 +1,635 @@
+// Copyright Epic Games, Inc. All Rights Reserved.
+
+#pragma once
+
+#include "CoreMinimal.h"
+#include "Misc/TVariant.h"
+#include "Online/OnlineMeta.h"
+#include "Online/OnlineResult.h"
+
+namespace UE::Online {
+
+/** Schema name - used to lookup and interact with a schema. */
+using FSchemaId = FName;
+/**
+  * Category within a schema.
+  * Most schema will have two sets of configuration - one for the singular object and one for
+  * member data within that object. The category name is expected to be built into the code which
+  * uses a schema and is not expected to change. A schema may not add any categories not already
+  * declared in the parent schema.
+  * 
+  * Example: A lobby contains both lobby data and lobby member data, but during lobby creation the
+  * application only specifies a single schema Id. The schema id is required to have categories for
+  * both Lobby and LobbyMember.
+  */
+using FSchemaCategoryId = FName;
+/**
+  * The id used to refer to a schema attribute. These ids are declared in the application config
+  * and used in application code to look up an attribute value.
+  */
+using FSchemaAttributeId = FName;
+/**
+  * The id used to refer to a schema service attribute. These ids are declared in config and used
+  * directly by the platform service when serializing attributes.
+  */
+using FSchemaServiceAttributeId = FName;
+/**
+  * The id used to refer to a schema service descriptor. Similar to a schema, service descriptors
+  * declare a set of attributes. Service descriptors are tied directly to how a service serializes
+  * data to its backend.
+  * 
+  * Example: In the case of a lobby the backing service is capable of storing attributes for a
+  * lobby and a lobby member. There will be separate service descriptors for lobby and lobby member
+  * where each will list the attributes available to be set as well as the types of data each
+  * attribute can store. The lobby implementation refers to the attributes names in the service
+  * descriptor directly when handling the serialization of attributes.
+  */
+using FSchemaServiceDescriptorId = FName;
+
+/**
+  * Flags to indicate required properties for a schema attribute.
+  */
+enum class ESchemaAttributeFlags : uint8
+{
+	/**
+	  * No special properties needed.
+	  */
+	None = 0,
+	/**
+	  * Attribute is available for search operations.
+	  * 
+	  * Searchable attributes are required to exist only in a base schema used as a parent of all
+	  * schema for a particular interface.
+	  */
+	Searchable = 1 << 0,
+	/**
+	  * Attribute is publicly visible.
+	  */
+	Public = 1 << 1,
+	/**
+	  * If multiple schema are supported using the same service interface, a base schema must exist
+	  * with a field marked as schema compatibility id. This field will allow remote clients to
+	  * determine what schema should be used to handle attributes.
+	  */
+	SchemaCompatibilityId = 1 << 2,
+};
+ENUM_CLASS_FLAGS(ESchemaAttributeFlags);
+ONLINESERVICESINTERFACE_API const TCHAR* LexToString(ESchemaAttributeFlags SchemaAttributeFlags);
+ONLINESERVICESINTERFACE_API void LexFromString(ESchemaAttributeFlags& OutSchemaAttributeFlags, const TCHAR* InStr);
+
+/**
+  * Flags to indicate the data types supported by a service attribute.
+  */
+enum class ESchemaAttributeType : uint8
+{
+	/**
+	  * Unset.
+	  */
+	None,
+	/**
+	  * Boolean values.
+	  */
+	Bool,
+	/**
+	  * Signed 64 bit integers.
+	  */
+	Int64,
+	/**
+	  * Double precision floating point.
+	  */
+	Double,
+	/**
+	  * Variable length string data.
+	  */
+	String,
+};
+ONLINESERVICESINTERFACE_API const TCHAR* LexToString(ESchemaAttributeType SchemaAttributeType);
+ONLINESERVICESINTERFACE_API void LexFromString(ESchemaAttributeType& OutSchemaAttributeType, const TCHAR* InStr);
+
+/**
+  * Flags to indicate the available support for a service attribute.
+  */
+enum class ESchemaServiceAttributeFlags : uint8
+{
+	/**
+	  * No special properties needed.
+	  */
+	None = 0,
+	/**
+	  * Attribute is available for search operations.
+	  */
+	Searchable = 1 << 0,
+	/**
+	  * Attribute is publicly visible.
+	  */
+	Public = 1 << 1,
+};
+ENUM_CLASS_FLAGS(ESchemaServiceAttributeFlags);
+ONLINESERVICESINTERFACE_API const TCHAR* LexToString(ESchemaServiceAttributeFlags SchemaServiceAttributeFlags);
+ONLINESERVICESINTERFACE_API void LexFromString(ESchemaServiceAttributeFlags& OutSchemaServiceAttributeFlags, const TCHAR* InStr);
+
+/**
+  * Flags to indicate the data types supported by a service attribute.
+  */
+enum class ESchemaServiceAttributeSupportedTypeFlags : uint8
+{
+	None = 0,
+	/**
+	  * Boolean values.
+	  */
+	Bool = 1 << 0,
+	/**
+	  * Signed 64 bit integers.
+	  */
+	Int64 = 1 << 1,
+	/**
+	  * Double precision floating point.
+	  */
+	Double = 1 << 2,
+	/**
+	  * Variable length string data.
+	  */
+	String = 1 << 3,
+};
+ENUM_CLASS_FLAGS(ESchemaServiceAttributeSupportedTypeFlags);
+ONLINESERVICESINTERFACE_API const TCHAR* LexToString(ESchemaServiceAttributeSupportedTypeFlags SchemaServiceAttributeSupportedTypeFlags);
+ONLINESERVICESINTERFACE_API void LexFromString(ESchemaServiceAttributeSupportedTypeFlags& OutSchemaServiceAttributeSupportedTypeFlags, const TCHAR* InStr);
+
+class ONLINESERVICESINTERFACE_API FSchemaVariant final
+{
+public:
+	using FVariantType = TVariant<FString, int64, double, bool>;
+	FSchemaVariant() = default;
+	FSchemaVariant(const FSchemaVariant& InOther) = default;
+	FSchemaVariant(FSchemaVariant&& InOther);
+	FSchemaVariant& operator=(FSchemaVariant&&);
+	FSchemaVariant& operator=(const FSchemaVariant&) = default;
+
+	template<typename ValueType>
+	FSchemaVariant(const ValueType& InData) { Set(InData); }
+	template<typename ValueType>
+	FSchemaVariant(ValueType&& InData) { Set(MoveTemp(InData)); }
+	void Set(const TCHAR* AsString);
+	void Set(const FString& AsString);
+	void Set(FString&& AsString);
+	void Set(int64 AsInt);
+	void Set(double AsDouble);
+	void Set(bool bAsBool);
+
+	ESchemaAttributeType GetType() const { return VariantType; }
+
+	int64 GetInt64() const;
+	double GetDouble() const;
+	bool GetBoolean() const;
+	FString GetString() const;
+
+	FString ToLogString() const;
+
+	bool operator==(const FSchemaVariant& Other) const;
+	bool operator!=(const FSchemaVariant& Other) const { return !(*this == Other); }
+public:
+	FVariantType VariantData;
+	ESchemaAttributeType VariantType = ESchemaAttributeType::None;
+};
+
+// Don't allow implicit conversion to FSchemaVariant when calling LexToString.
+template <typename T> class FLexToStringAdaptor;
+template <>
+class FLexToStringAdaptor<FSchemaVariant>
+{
+public:
+	FLexToStringAdaptor(const FSchemaVariant& SchemaVariant)
+		: SchemaVariant(SchemaVariant)
+	{
+	}
+
+	const FSchemaVariant& SchemaVariant;
+};
+ONLINESERVICESINTERFACE_API const FString LexToString(const FLexToStringAdaptor<FSchemaVariant>& Adaptor);
+ONLINESERVICESINTERFACE_API void LexFromString(FSchemaVariant& OutSchemaVariant, const TCHAR* InStr);
+
+/**
+  * Describes available options for an attribute made available by the underlying service.
+  */
+struct FSchemaServiceAttributeDescriptor
+{
+	/**
+	  * The id to associate with this attribute.
+	  */
+	FSchemaServiceAttributeId Id;
+	/**
+	  * Declares what data types can be handled by the service attribute.
+	  */
+	TArray<ESchemaServiceAttributeSupportedTypeFlags> SupportedTypes;
+	/**
+	  * Additional behavior required by the attribute. Examples include making an attribute
+	  * public or available to search.
+	  */
+	TArray<ESchemaServiceAttributeFlags> Flags;
+	/**
+	  * Declares the maximum allowed size for the service attribute.
+	  *
+	  * Applicable when using a variable sized type such as String.
+	  */
+	int32 MaxSize = 0;
+};
+
+/**
+  * Service descriptor used for declaring a set of attributes a service provides.
+  */
+struct FSchemaServiceDescriptor
+{
+	/**
+	  * The id associated with this service descriptor.
+	  */
+	FSchemaServiceDescriptorId Id;
+	/**
+	  * Service attribute descriptions.
+	  */
+	TArray<FSchemaServiceAttributeDescriptor> Attributes;
+};
+
+/**
+  * Describes the required attribute configuration for a schema attribute.
+  */
+struct FSchemaAttributeDescriptor
+{
+	/**
+	  * The id to associate with this attribute.
+	  */
+	FSchemaAttributeId Id;
+	/**
+	  * The type of data contained within the attribute.
+	  */
+	ESchemaAttributeType Type = ESchemaAttributeType::None;
+	/**
+	  * Additional behavior required by the attribute. Examples include making an attribute
+	  * public or available to search.
+	  */
+	TArray<ESchemaAttributeFlags> Flags;
+	/**
+	  * Whether to group the attribute together with other attributes explicitly.
+	  * Providing a group id allows the application to explicitly serialize attributes to the same
+	  * service attribute for increased network efficiency.
+	  *
+	  * Id 0 uses automatic grouping.
+	 */
+	int32 UpdateGroupId = 0;
+	/**
+	  * Declares the maximum allowed serialized size for a variably sized attribute.
+	  *
+	  * Applicable when using a variable sized type such as String.
+	  */
+	int32 MaxSize = 0;
+};
+
+struct FSchemaCategoryDescriptor
+{
+	/**
+	  * Category for which the schema applies.
+	  * 
+	  * Example: "Lobby" vs "LobbyMember"
+	  */
+	FSchemaCategoryId Id;
+	/**
+	  * Service descriptor to use for initializing schema.
+	  * May be empty when a schema only applies to local data.
+	  * May only be set by base parent schema.
+	  */
+	FSchemaServiceDescriptorId ServiceDescriptorId;
+	/**
+	  * Attribute definitions.
+	  */
+	TArray<FSchemaAttributeDescriptor> Attributes;
+};
+
+/**
+  * Describes all properties needed to use a schema.
+  */
+struct FSchemaDescriptor
+{
+	/**
+	  * The id associated with this schema.
+	  */
+	FSchemaId Id;
+	/**
+	  * Parent schema (not required).
+	  * Attributes used in a parent schema may not be modified by a child schema.
+	  */
+	FSchemaId ParentId;
+	/**
+	  * Categories of data within the schema. Each category is handled differently within the
+	  * service so will have different configuration of attributes and service attributes.
+	  */
+	TArray<FSchemaCategoryDescriptor> Categories;
+};
+
+/**
+  * Set of all descriptors used to initialize a SchemaRegistry.
+  */
+struct FSchemaRegistryDescriptorConfig
+{
+	/**
+	  * List of all supported schema.
+	  */
+	TArray<FSchemaDescriptor> SchemaDescriptors;
+	/**
+	  * List of all service descriptors.
+	  */
+	TArray<FSchemaServiceDescriptor> ServiceDescriptors;
+};
+
+/**
+  * The runtime definition of a schema service attribute.
+  */
+struct FSchemaServiceAttributeDefinition
+{
+	/**
+	  * The id to associate with this attribute.
+	  */
+	FSchemaServiceAttributeId Id;
+	/**
+	  * The supported type. Only one type is supported in a service attribute definition.
+	  */
+	ESchemaServiceAttributeSupportedTypeFlags Type;
+	/**
+	  * Declares the maximum allowed size for the service attribute.
+	  *
+	  * Applicable when using a variable sized type such as String.
+	  */
+	ESchemaAttributeFlags Flags = ESchemaAttributeFlags::None;
+	/**
+	  * Declares the maximum allowed size for the service attribute.
+	  *
+	  * Applicable when using a variable sized type such as String.
+	  */
+	int32 MaxSize = 0;
+	/**
+	  * The attribute ids contained within the service attribute.
+	  */
+	TArray<FSchemaAttributeId> SchemaAttributeIds;
+};
+
+/**
+  * The runtime definition of a schema attribute.
+  */
+struct FSchemaAttributeDefinition
+{
+	/**
+	  * The id to associate with this attribute.
+	  */
+	FSchemaAttributeId Id;
+	/**
+	  * The type of data contained within the attribute.
+	  */
+	ESchemaAttributeType Type = ESchemaAttributeType::None;
+	/**
+	  * Additional behavior required by the attribute. Examples include making an attribute
+	  * public or available to search.
+	  */
+	ESchemaAttributeFlags Flags;
+	/**
+	  * Declares the maximum allowed serialized size for a variably sized attribute.
+	  *
+	  * Applicable when using a variable sized type such as String.
+	  */
+	int32 MaxSize = 0;
+	/**
+	  * The id of the associated service attribute.
+	  */
+	FSchemaServiceAttributeId ServiceAttributeId;
+};
+
+struct FSchemaCategoryDefinition
+{
+	/**
+	  * Category for which the schema applies.
+	  *
+	  * Example: "Lobby" vs "LobbyMember"
+	  */
+	FSchemaCategoryId Id;
+	/**
+	  * Mapping of schema attribute ids to attribute definitions.
+	  */
+	TMap<FSchemaAttributeId, FSchemaAttributeDefinition> SchemaAttributeDefinitions;
+	/**
+	  * Mapping of service attribute ids to service attribute definitions.
+	  */
+	TMap<FSchemaServiceAttributeId, FSchemaServiceAttributeDefinition> ServiceAttributeDefinitions;
+};
+
+/**
+  * The runtime definition of a schema.
+  * The contained data can be used for validation and translation.
+  */
+struct FSchemaDefinition
+{
+	/**
+	  * The id associated with this schema.
+	  */
+	FSchemaId Id;
+	/**
+	  * The compatibility id consists of the CRC of the schema name in the top 32 bits and the CRC
+	  * of all the schema options in the bottom 32 bits.
+	  */
+	int64 CompatibilityId = 0;
+	/**
+	  * Supported categories within the schema.
+	  */
+	TMap<FSchemaCategoryId, FSchemaCategoryDefinition> Categories;
+	/**
+	  * Set of parent schema ids for parent checks.
+	  */
+	TSet<FSchemaId> ParentSchemaIds;
+};
+
+/**
+  * The runtime definition of a schema.
+  * The contained data can be used for validation and translation.
+  */
+struct FSchemaServiceAttributeData
+{
+	/**
+	  * The id to associate with this attribute.
+	  */
+	FSchemaServiceAttributeId Id;
+	/**
+	  * Additional behavior required by the attribute. Examples include making an attribute
+	  * public or available to search.
+	  */
+	ESchemaAttributeFlags Flags = ESchemaAttributeFlags::None;
+	/**
+	  * The attribute data.
+	  */
+	FSchemaVariant Data;
+};
+
+/**
+  * The runtime definition of a schema.
+  * The contained data can be used for validation and translation.
+  */
+struct FSchemaAttributeData
+{
+	/**
+	  * The id to associate with this attribute.
+	  */
+	FSchemaAttributeId Id;
+	/**
+	  * The attribute data.
+	  */
+	FSchemaVariant Data;
+};
+
+/**
+  * Changes to be applied to the service data following schema attribute changes.
+  */
+struct FSchemaServiceChanges
+{
+	/**
+	  * Added attributes and their values.
+	  */
+	TMap<FSchemaServiceAttributeId, FSchemaServiceAttributeData> AddedAttributes;
+	/**
+	  * Removed attributes with their previous values.
+	  */
+	TSet<FSchemaAttributeId> RemovedAttributes;
+	/**
+	  * Changed attributes with their old and new values.
+	  */
+	TMap<FSchemaServiceAttributeId, TPair<FSchemaServiceAttributeData, FSchemaServiceAttributeData>> ChangedAttributes;
+};
+
+/**
+  * Changes to be applied to the application data following service attribute changes.
+  */
+struct FSchemaApplicationChanges
+{
+	/**
+	  * Added attributes and their values.
+	  */
+	TMap<FSchemaAttributeId, FSchemaAttributeData> AddedAttributes;
+	/**
+	  * Removed attributes.
+	  */
+	TSet<FSchemaAttributeId> RemovedAttributes;
+	/**
+	  * Changed attributes with their old and new values.
+	  */
+	TMap<FSchemaAttributeId, TPair<FSchemaAttributeData, FSchemaAttributeData>> ChangedAttributes;
+};
+
+#if 0 // todo
+struct FSchemaCategoryInstanceChangeSchema
+{
+	static constexpr TCHAR Name[] = TEXT("ChangeSchema");
+
+	struct Params
+	{
+		/** Schema Id to switch to. The schema must be configured and be a sibling of the current
+		  * schema.
+		  * A parent schema field definition for SchemaId must exist to enable this feature.
+		  */
+		FSchemaId SchemaId;
+	};
+
+	struct Result
+	{
+		/** Attribute changes to be notified to the service. */
+		TSharedRef<FSchemaServiceChanges> ServiceChanges;
+
+	public:
+		Result() = delete; // cannot default construct due to TSharedRef
+	};
+};
+#endif
+
+struct FSchemaCategoryInstanceApplyApplicationChanges
+{
+	static constexpr TCHAR Name[] = TEXT("ApplyApplicationChanges");
+
+	struct Params
+	{
+		/** Attributes notified from the application. */
+		TArray<FSchemaAttributeData> MutatedAttributes;
+		/** Removed attributes notified from the application. */
+		TArray<FSchemaAttributeId> RemovedAttributes;
+		/** Whether the attributes represent a full snapshot or a delta from a prior update. */
+		bool bIsDelta = false;
+	};
+
+	struct Result
+	{
+		/** Attribute changes to be notified to the service. */
+		FSchemaServiceChanges ServiceChanges;
+
+	public:
+		Result() = delete; // cannot default construct due to TSharedRef
+	};
+};
+
+struct FSchemaCategoryInstanceApplyServiceChanges
+{
+	static constexpr TCHAR Name[] = TEXT("ApplyServiceChanges");
+
+	struct Params
+	{
+		/** Attributes notified from the service. */
+		TArray<FSchemaServiceAttributeData> MutatedAttributes;
+		/** Removed attributes notified from the service. */
+		TArray<FSchemaAttributeId> RemovedAttributes;
+		/** Whether the attributes represent a full snapshot or a delta from a prior update. */
+		bool bIsDelta = false;
+	};
+
+	struct Result
+	{
+		/** Attribute changes to be notified to the application. */
+		FSchemaApplicationChanges ApplicationChanges;
+
+	public:
+		Result() = delete; // cannot default construct due to TSharedRef
+	};
+};
+
+namespace Meta {
+
+BEGIN_ONLINE_STRUCT_META(FSchemaServiceAttributeDescriptor)
+	ONLINE_STRUCT_FIELD(FSchemaServiceAttributeDescriptor, Id),
+	ONLINE_STRUCT_FIELD(FSchemaServiceAttributeDescriptor, SupportedTypes),
+	ONLINE_STRUCT_FIELD(FSchemaServiceAttributeDescriptor, Flags),
+	ONLINE_STRUCT_FIELD(FSchemaServiceAttributeDescriptor, MaxSize)
+END_ONLINE_STRUCT_META()
+
+BEGIN_ONLINE_STRUCT_META(FSchemaServiceDescriptor)
+	ONLINE_STRUCT_FIELD(FSchemaServiceDescriptor, Id),
+	ONLINE_STRUCT_FIELD(FSchemaServiceDescriptor, Attributes)
+END_ONLINE_STRUCT_META()
+
+BEGIN_ONLINE_STRUCT_META(FSchemaAttributeDescriptor)
+	ONLINE_STRUCT_FIELD(FSchemaAttributeDescriptor, Id),
+	ONLINE_STRUCT_FIELD(FSchemaAttributeDescriptor, Type),
+	ONLINE_STRUCT_FIELD(FSchemaAttributeDescriptor, Flags),
+	ONLINE_STRUCT_FIELD(FSchemaAttributeDescriptor, UpdateGroupId),
+	ONLINE_STRUCT_FIELD(FSchemaAttributeDescriptor, MaxSize)
+END_ONLINE_STRUCT_META()
+
+BEGIN_ONLINE_STRUCT_META(FSchemaCategoryDescriptor)
+	ONLINE_STRUCT_FIELD(FSchemaCategoryDescriptor, Id),
+	ONLINE_STRUCT_FIELD(FSchemaCategoryDescriptor, ServiceDescriptorId),
+	ONLINE_STRUCT_FIELD(FSchemaCategoryDescriptor, Attributes)
+END_ONLINE_STRUCT_META()
+
+BEGIN_ONLINE_STRUCT_META(FSchemaDescriptor)
+	ONLINE_STRUCT_FIELD(FSchemaDescriptor, Id),
+	ONLINE_STRUCT_FIELD(FSchemaDescriptor, ParentId),
+	ONLINE_STRUCT_FIELD(FSchemaDescriptor, Categories)
+END_ONLINE_STRUCT_META()
+
+BEGIN_ONLINE_STRUCT_META(FSchemaRegistryDescriptorConfig)
+	ONLINE_STRUCT_FIELD(FSchemaRegistryDescriptorConfig, SchemaDescriptors),
+	ONLINE_STRUCT_FIELD(FSchemaRegistryDescriptorConfig, ServiceDescriptors)
+END_ONLINE_STRUCT_META()
+
+/* Meta*/ }
+
+/* UE::Online */ }
