Turned on unity builds for OnlineServicesCommon and OnlineServicesInterface.

Compiling modules directly:
Total CPU Time Before:
OnlineServiceCommon: 177.984375 s
OnlineServiceInterface:  17.71875 s

Total CPU Time After:
OnlineServicesCommon: 88.34375 s
OnlineServicesInterface: 8 s

diff --git a/Engine/Plugins/Online/Experimental/OnlineServices/Source/OnlineServicesCommon/OnlineServicesCommon.Build.cs b/Engine/Plugins/Online/Experimental/OnlineServices/Source/OnlineServicesCommon/OnlineServicesCommon.Build.cs
--- a/Engine/Plugins/Online/Experimental/OnlineServices/Source/OnlineServicesCommon/OnlineServicesCommon.Build.cs
+++ b/Engine/Plugins/Online/Experimental/OnlineServices/Source/OnlineServicesCommon/OnlineServicesCommon.Build.cs
@@ -8,8 +8,6 @@ public class OnlineServicesCommon : ModuleRules
 	{
 		PCHUsage = ModuleRules.PCHUsageMode.UseExplicitOrSharedPCHs;
 
-		bUseUnity = false; // set false to catch non-unity compile issues while prototyping
-
 		PublicDependencyModuleNames.AddRange(
 			new string[] {
 				"CoreOnline",
diff --git a/Engine/Plugins/Online/Experimental/OnlineServices/Source/OnlineServicesInterface/OnlineServicesInterface.Build.cs b/Engine/Plugins/Online/Experimental/OnlineServices/Source/OnlineServicesInterface/OnlineServicesInterface.Build.cs
--- a/Engine/Plugins/Online/Experimental/OnlineServices/Source/OnlineServicesInterface/OnlineServicesInterface.Build.cs
+++ b/Engine/Plugins/Online/Experimental/OnlineServices/Source/OnlineServicesInterface/OnlineServicesInterface.Build.cs
@@ -8,8 +8,6 @@ public class OnlineServicesInterface : ModuleRules
     {
 		PCHUsage = ModuleRules.PCHUsageMode.UseExplicitOrSharedPCHs;
 
-		bUseUnity = false; // set false to catch non-unity compile issues while prototyping
-
 		PublicDependencyModuleNames.AddRange(
 			new string[] {
 				"CoreOnline",
