Quick fix for Exec Handler's logic regarding Maps

diff --git a/Engine/Plugins/Online/Experimental/OnlineServices/Source/OnlineServicesCommon/Public/Online/OnlineExecHandler.h b/Engine/Plugins/Online/Experimental/OnlineServices/Source/OnlineServicesCommon/Public/Online/OnlineExecHandler.h
--- a/Engine/Plugins/Online/Experimental/OnlineServices/Source/OnlineServicesCommon/Public/Online/OnlineExecHandler.h
+++ b/Engine/Plugins/Online/Experimental/OnlineServices/Source/OnlineServicesCommon/Public/Online/OnlineExecHandler.h
@@ -594,7 +594,7 @@ inline bool ParseOnlineExecParams(const TCHAR*& Cmd, TMap<T, U>& Map, IOnlineSer
 			TArray<FString> TokenValuePair;
 			ArrayToken.ParseIntoArray(TokenValuePair, TEXT("="));
 
-			if (TokenValuePair.Num() > 2)
+			if (TokenValuePair.Num() == 2)
 			{
 				T Key;
 				U Value;
@@ -603,7 +603,7 @@ inline bool ParseOnlineExecParams(const TCHAR*& Cmd, TMap<T, U>& Map, IOnlineSer
 				const TCHAR* ArrayTokenKeyTCHAR = *TrimmedKey;
 				if (ParseOnlineExecParams(ArrayTokenKeyTCHAR, Key, Services))
 				{
-					FString TrimmedValue = TokenValuePair[2].TrimStartAndEnd();
+					FString TrimmedValue = TokenValuePair[1].TrimStartAndEnd();
 					const TCHAR* ArrayTokenValueTCHAR = *TrimmedValue;
 					if (ParseOnlineExecParams(ArrayTokenValueTCHAR, Value, Services))
 					{
