OSSv2: Add Leaderboards integration tests.

diff --git a/Engine/Plugins/Online/Experimental/OnlineServices/Source/OnlineServicesCommon/Private/Online/AchievementsCommon.cpp b/Engine/Plugins/Online/Experimental/OnlineServices/Source/OnlineServicesCommon/Private/Online/AchievementsCommon.cpp
--- a/Engine/Plugins/Online/Experimental/OnlineServices/Source/OnlineServicesCommon/Private/Online/AchievementsCommon.cpp
+++ b/Engine/Plugins/Online/Experimental/OnlineServices/Source/OnlineServicesCommon/Private/Online/AchievementsCommon.cpp
@@ -23,7 +23,7 @@ void FAchievementsCommon::Initialize()
 {
 	TOnlineComponent<IAchievements>::Initialize();
 
-	StatEventHandle = Services.Get<FStatsCommon>()->OnStatsUpdated().Add([this](const FStatsUpdated& StatsUpdated) { UnlockAchievementsByStats(StatsUpdated); });
+	StatEventHandle = Services.Get<IStats>()->OnStatsUpdated().Add([this](const FStatsUpdated& StatsUpdated) { UnlockAchievementsByStats(StatsUpdated); });
 }
 
 void FAchievementsCommon::Shutdown()
diff --git a/Engine/Plugins/Online/Experimental/OnlineServices/Source/OnlineServicesCommon/Private/Online/LeaderboardsCommon.cpp b/Engine/Plugins/Online/Experimental/OnlineServices/Source/OnlineServicesCommon/Private/Online/LeaderboardsCommon.cpp
--- a/Engine/Plugins/Online/Experimental/OnlineServices/Source/OnlineServicesCommon/Private/Online/LeaderboardsCommon.cpp
+++ b/Engine/Plugins/Online/Experimental/OnlineServices/Source/OnlineServicesCommon/Private/Online/LeaderboardsCommon.cpp
@@ -4,6 +4,7 @@
 
 #include "Online/Auth.h"
 #include "Online/OnlineServicesCommon.h"
+#include "Online/StatsCommon.h"
 
 namespace UE::Online {
 
@@ -66,15 +67,23 @@ void LexFromString(ELeaderboardOrderMethod& OutValue, const TCHAR* InStr)
 }
 
 FLeaderboardsCommon::FLeaderboardsCommon(FOnlineServicesCommon& InServices)
-	: TOnlineComponent(TEXT("Leaderboards"), InServices)
+	: Super(TEXT("Leaderboards"), InServices)
 {
 }
 
+void FLeaderboardsCommon::Initialize()
+{
+	Super::Initialize();
+
+	StatEventHandle = Services.Get<IStats>()->OnStatsUpdated().Add([this](const FStatsUpdated& StatsUpdated) { WriteLeaderboardsByStats(StatsUpdated); });
+}
+
 void FLeaderboardsCommon::LoadConfig()
 {
 	Super::LoadConfig();
 
 	const TCHAR* ConfigSection = TEXT("OnlineServices.Leaderboards");
+	GConfig->GetBool(ConfigSection, TEXT("LeaderboardDef_IsTitleManaged"), bIsTitleManaged, GEngineIni);
 
 	for (int StatIdx = 0;; StatIdx++)
 	{
@@ -143,4 +152,27 @@ TOnlineAsyncOpHandle<FWriteLeaderboardScores> FLeaderboardsCommon::WriteLeaderbo
 	return Operation->GetHandle();
 }
 
+void FLeaderboardsCommon::WriteLeaderboardsByStats(const FStatsUpdated& StatsUpdated)
+{
+	if (!bIsTitleManaged)
+	{
+		return;
+	}
+
+	for (const FUserStats& UserStats : StatsUpdated.UpdateUsersStats)
+	{
+		for (const TPair<FString, FStatValue>& StatPair : UserStats.Stats)
+		{
+			if (LeaderboardDefinitions.Contains(StatPair.Key))
+			{
+				FWriteLeaderboardScores::Params WriteLeaderboardScoresParam;
+				WriteLeaderboardScoresParam.LocalUserId = UserStats.UserId;
+				WriteLeaderboardScoresParam.BoardName = StatPair.Key;
+				WriteLeaderboardScoresParam.Score = StatPair.Value.GetInt64();
+				WriteLeaderboardScores(MoveTemp(WriteLeaderboardScoresParam));
+			}
+		}
+	}
+}
+
 /* UE::Online */ }
diff --git a/Engine/Plugins/Online/Experimental/OnlineServices/Source/OnlineServicesCommon/Public/Online/AchievementsCommon.h b/Engine/Plugins/Online/Experimental/OnlineServices/Source/OnlineServicesCommon/Public/Online/AchievementsCommon.h
--- a/Engine/Plugins/Online/Experimental/OnlineServices/Source/OnlineServicesCommon/Public/Online/AchievementsCommon.h
+++ b/Engine/Plugins/Online/Experimental/OnlineServices/Source/OnlineServicesCommon/Public/Online/AchievementsCommon.h
@@ -50,7 +50,6 @@ public:
 protected:
 	TOnlineEventCallable<void(const FAchievementStateUpdated&)> OnAchievementStateUpdatedEvent;
 
-	bool IsTitleManaged() const;
 	void UnlockAchievementsByStats(const FStatsUpdated& StatsUpdated);
 	void ExecuteUnlockRulesRelatedToStat(const FOnlineAccountIdHandle& UserId, const FString& StatName, const TMap<FString, FStatValue>& Stats, TArray<FString>& OutAchievementsToUnlock);
 	bool MeetUnlockCondition(FAchievementUnlockRule AchievementUnlockRule, const TMap<FString, FStatValue>& Stats);
diff --git a/Engine/Plugins/Online/Experimental/OnlineServices/Source/OnlineServicesCommon/Public/Online/LeaderboardsCommon.h b/Engine/Plugins/Online/Experimental/OnlineServices/Source/OnlineServicesCommon/Public/Online/LeaderboardsCommon.h
--- a/Engine/Plugins/Online/Experimental/OnlineServices/Source/OnlineServicesCommon/Public/Online/LeaderboardsCommon.h
+++ b/Engine/Plugins/Online/Experimental/OnlineServices/Source/OnlineServicesCommon/Public/Online/LeaderboardsCommon.h
@@ -4,6 +4,7 @@
 
 #include "Online/Leaderboards.h"
 #include "Online/OnlineComponent.h"
+#include "Online/Stats.h"
 
 namespace UE::Online {
 
@@ -64,6 +65,7 @@ public:
 	FLeaderboardsCommon(FOnlineServicesCommon& InServices);
 
 	// TOnlineComponent
+	virtual void Initialize() override;
 	virtual void LoadConfig() override;
 	virtual void RegisterCommands() override;
 
@@ -76,7 +78,11 @@ public:
 	virtual TOnlineAsyncOpHandle<FWriteLeaderboardScores> WriteLeaderboardScores(FWriteLeaderboardScores::Params&& Params);
 
 protected:
+	void WriteLeaderboardsByStats(const FStatsUpdated& StatsUpdated);
+
 	TMap<FString, FLeaderboardDefinition> LeaderboardDefinitions;
+	FOnlineEventDelegateHandle StatEventHandle;
+	bool bIsTitleManaged = false;
 };
 
 namespace Meta {
diff --git a/Engine/Plugins/Online/Experimental/OnlineServices/Source/OnlineServicesCommon/Public/Online/OnlineComponent.h b/Engine/Plugins/Online/Experimental/OnlineServices/Source/OnlineServicesCommon/Public/Online/OnlineComponent.h
--- a/Engine/Plugins/Online/Experimental/OnlineServices/Source/OnlineServicesCommon/Public/Online/OnlineComponent.h
+++ b/Engine/Plugins/Online/Experimental/OnlineServices/Source/OnlineServicesCommon/Public/Online/OnlineComponent.h
@@ -15,6 +15,8 @@ class TOnlineComponent
 	, public IOnlineComponent
 {
 public:
+	using Super = ComponentType;
+
 	TOnlineComponent(const TOnlineComponent&) = delete;
 	TOnlineComponent(TOnlineComponent&&) = delete;
 
diff --git a/Engine/Plugins/Online/Experimental/OnlineServicesNull/Source/Private/Online/LeaderboardsNull.cpp b/Engine/Plugins/Online/Experimental/OnlineServicesNull/Source/Private/Online/LeaderboardsNull.cpp
--- a/Engine/Plugins/Online/Experimental/OnlineServicesNull/Source/Private/Online/LeaderboardsNull.cpp
+++ b/Engine/Plugins/Online/Experimental/OnlineServicesNull/Source/Private/Online/LeaderboardsNull.cpp
@@ -103,6 +103,7 @@ TOnlineAsyncOpHandle<FWriteLeaderboardScores> FLeaderboardsNull::WriteLeaderboar
 		UserScoreToInsert.Score = UpdatedScore;
 		UserScoreToInsert.UserId = InAsyncOp.GetParams().LocalUserId;
 
+		bool bInserted = false;
 		CurrentNode = UserScoreList.GetHead();
 		while (CurrentNode)
 		{
@@ -122,15 +123,17 @@ TOnlineAsyncOpHandle<FWriteLeaderboardScores> FLeaderboardsNull::WriteLeaderboar
 			}
 			if (ShouldInsert)
 			{
+				bInserted = true;
 				UserScoreList.InsertNode(UserScoreToInsert, CurrentNode);
 				break;
 			}
 
 			CurrentNode = CurrentNode->GetNextNode();
-			if (!CurrentNode)
-			{
-				UserScoreList.AddTail(UserScoreToInsert);
-			}
+		}
+
+		if (!bInserted)
+		{
+			UserScoreList.AddTail(UserScoreToInsert);
 		}
 
 		InAsyncOp.SetResult({});
