Quick fixes for Sessions functionality

diff --git a/Engine/Plugins/Online/Experimental/OnlineServices/Source/OnlineServicesCommon/Private/Online/SessionsCommon.cpp b/Engine/Plugins/Online/Experimental/OnlineServices/Source/OnlineServicesCommon/Private/Online/SessionsCommon.cpp
--- a/Engine/Plugins/Online/Experimental/OnlineServices/Source/OnlineServicesCommon/Private/Online/SessionsCommon.cpp
+++ b/Engine/Plugins/Online/Experimental/OnlineServices/Source/OnlineServicesCommon/Private/Online/SessionsCommon.cpp
@@ -504,9 +504,23 @@ namespace UE::Online {
 
 	FOnlineError FSessionsCommon::CheckCreateSessionParams(const FCreateSession::Params& Params)
 	{
+		if (Params.SessionName.IsNone())
+		{
+			UE_LOG(LogTemp, Warning, TEXT("[FSessionsCommon::CheckCreateSessionParams] Could not create session with no valid SessionName set"));
+
+			return Errors::InvalidParams();
+		}
+
+		if (!Params.LocalUserId.IsValid())
+		{
+			UE_LOG(LogTemp, Warning, TEXT("[FSessionsCommon::CheckCreateSessionParams] Could not create session with name [%s]. LocalUserId [%s] not valid"), *Params.SessionName.ToString(), *ToLogString(Params.LocalUserId));
+
+			return Errors::InvalidParams();
+		}
+
 		if (Params.SessionSettings.NumMaxPrivateConnections == 0 && Params.SessionSettings.NumMaxPublicConnections == 0)
 		{
-			UE_LOG(LogTemp, Warning, TEXT("[FSessionsCommon::CheckCreateSessionParams] Could not create session with no valid NumMaxPrivateConnections [%d] or NumMaxPublicConnections [%d]"), Params.SessionSettings.NumMaxPrivateConnections, Params.SessionSettings.NumMaxPublicConnections);
+			UE_LOG(LogTemp, Warning, TEXT("[FSessionsCommon::CheckCreateSessionParams] Could not create session with name [%s] with no valid NumMaxPrivateConnections [%d] or NumMaxPublicConnections [%d]"), *Params.SessionName.ToString(), Params.SessionSettings.NumMaxPrivateConnections, Params.SessionSettings.NumMaxPublicConnections);
 
 			return Errors::InvalidParams();
 		}
@@ -680,7 +694,7 @@ namespace UE::Online {
 			return Errors::InvalidParams();
 		}
 
-		if (Params.SessionName.ToString().IsEmpty())
+		if (Params.SessionName.IsNone())
 		{
 			UE_LOG(LogTemp, Warning, TEXT("[FSessionsCommon::CheckJoinSessionParams] Could not join session with no valid SessionName set"));
 
diff --git a/Engine/Plugins/Online/Experimental/OnlineServices/Source/OnlineServicesCommonEngineUtils/Private/Online/SessionsLAN.cpp b/Engine/Plugins/Online/Experimental/OnlineServices/Source/OnlineServicesCommonEngineUtils/Private/Online/SessionsLAN.cpp
--- a/Engine/Plugins/Online/Experimental/OnlineServices/Source/OnlineServicesCommonEngineUtils/Private/Online/SessionsLAN.cpp
+++ b/Engine/Plugins/Online/Experimental/OnlineServices/Source/OnlineServicesCommonEngineUtils/Private/Online/SessionsLAN.cpp
@@ -130,16 +130,6 @@ TOnlineAsyncOpHandle<FCreateSession> FSessionsLAN::CreateSession(FCreateSession:
 		return Op->GetHandle();
 	}
 
-	// We'll only host on the LAN beacon for public sessions
-	if (OpParams.SessionSettings.JoinPolicy == ESessionJoinPolicy::Public)
-	{
-		if (TOptional<FOnlineError> TryHostLANSessionResult = TryHostLANSession())
-		{
-			Op->SetError(MoveTemp(TryHostLANSessionResult.GetValue()));
-			return Op->GetHandle();
-		}
-	}
-
 	Op->Then([this](TOnlineAsyncOp<FCreateSession>& Op)
 	{
 		const FCreateSession::Params& OpParams = Op.GetParams();
@@ -151,6 +141,16 @@ TOnlineAsyncOpHandle<FCreateSession> FSessionsLAN::CreateSession(FCreateSession:
 			return;
 		}
 
+		// We'll only host on the LAN beacon for public sessions
+		if (OpParams.SessionSettings.JoinPolicy == ESessionJoinPolicy::Public)
+		{
+			if (TOptional<FOnlineError> TryHostLANSessionResult = TryHostLANSession())
+			{
+				Op.SetError(MoveTemp(TryHostLANSessionResult.GetValue()));
+				return;
+			}
+		}
+
 		TSharedRef<FSessionLAN> NewSessionLANRef = MakeShared<FSessionLAN>();
 		NewSessionLANRef->CurrentState = ESessionState::Valid;
 		NewSessionLANRef->OwnerUserId = OpParams.LocalUserId;
diff --git a/Engine/Plugins/Online/Experimental/OnlineServices/Source/OnlineServicesInterface/Public/Online/Sessions.h b/Engine/Plugins/Online/Experimental/OnlineServices/Source/OnlineServicesInterface/Public/Online/Sessions.h
--- a/Engine/Plugins/Online/Experimental/OnlineServices/Source/OnlineServicesInterface/Public/Online/Sessions.h
+++ b/Engine/Plugins/Online/Experimental/OnlineServices/Source/OnlineServicesInterface/Public/Online/Sessions.h
@@ -152,10 +152,10 @@ struct ONLINESERVICESINTERFACE_API FSessionSettings
 	uint32 NumOpenPrivateConnections = 0;
 
 	/* Enum value describing the level of restriction to join the session */
-	ESessionJoinPolicy JoinPolicy;
+	ESessionJoinPolicy JoinPolicy = ESessionJoinPolicy::Public;
 
- 	/* In platforms that support this feature, it will set the session id to this value. Might be subject to minimum and maximum length */
- 	FString SessionIdOverride;
+	/* In platforms that support this feature, it will set the session id to this value. Might be subject to minimum and maximum length */
+	FString SessionIdOverride;
 
 	/* Whether the session is only available in the local network and not via internet connection. Only available in some platforms. False by default */
 	bool bIsLANSession = false;
@@ -840,9 +840,15 @@ public:
 namespace Meta {
 
 BEGIN_ONLINE_STRUCT_META(FFindSessionsSearchFilter)
+	ONLINE_STRUCT_FIELD(FFindSessionsSearchFilter, Key),
+	ONLINE_STRUCT_FIELD(FFindSessionsSearchFilter, ComparisonOp),
+	ONLINE_STRUCT_FIELD(FFindSessionsSearchFilter, Value)
 END_ONLINE_STRUCT_META()
 
 BEGIN_ONLINE_STRUCT_META(FCustomSessionSetting)
+	ONLINE_STRUCT_FIELD(FCustomSessionSetting, Data),
+	ONLINE_STRUCT_FIELD(FCustomSessionSetting, Visibility),
+	ONLINE_STRUCT_FIELD(FCustomSessionSetting, ID)
 END_ONLINE_STRUCT_META()
 
 BEGIN_ONLINE_STRUCT_META(FSessionMember)
