- Fix bIsHidden missing from the meta

diff --git a/Engine/Plugins/Online/Experimental/OnlineServices/Source/OnlineServicesInterface/Public/Online/Achievements.h b/Engine/Plugins/Online/Experimental/OnlineServices/Source/OnlineServicesInterface/Public/Online/Achievements.h
--- a/Engine/Plugins/Online/Experimental/OnlineServices/Source/OnlineServicesInterface/Public/Online/Achievements.h
+++ b/Engine/Plugins/Online/Experimental/OnlineServices/Source/OnlineServicesInterface/Public/Online/Achievements.h
@@ -267,6 +267,7 @@ BEGIN_ONLINE_STRUCT_META(FAchievementDefinition)
 	ONLINE_STRUCT_FIELD(FAchievementDefinition, FlavorText),
 	ONLINE_STRUCT_FIELD(FAchievementDefinition, UnlockedIconUrl),
 	ONLINE_STRUCT_FIELD(FAchievementDefinition, LockedIconUrl),
+	ONLINE_STRUCT_FIELD(FAchievementDefinition, bIsHidden),
 	ONLINE_STRUCT_FIELD(FAchievementDefinition, StatDefinitions)
 END_ONLINE_STRUCT_META()
 
