Add presence metadata support for enums to fix build issue.

diff --git a/Engine/Plugins/Online/Experimental/OnlineServices/Source/OnlineServicesInterface/Private/Online/Presence.cpp b/Engine/Plugins/Online/Experimental/OnlineServices/Source/OnlineServicesInterface/Private/Online/Presence.cpp
new file mode 100644
--- /dev/null
+++ b/Engine/Plugins/Online/Experimental/OnlineServices/Source/OnlineServicesInterface/Private/Online/Presence.cpp
@@ -0,0 +1,128 @@
+// Copyright Epic Games, Inc. All Rights Reserved.
+
+#include "Online/Presence.h"
+
+namespace UE::Online {
+
+const TCHAR* LexToString(EUserPresenceStatus EnumVal)
+{
+	switch (EnumVal)
+	{
+	case EUserPresenceStatus::Offline:		return TEXT("Offline");
+	case EUserPresenceStatus::Online:		return TEXT("Online");
+	case EUserPresenceStatus::Away:			return TEXT("Away");
+	case EUserPresenceStatus::ExtendedAway:	return TEXT("ExtendedAway");
+	case EUserPresenceStatus::DoNotDisturb:	return TEXT("DoNotDisturb");
+	default:								checkNoEntry(); // Intentional fall-through
+	case EUserPresenceStatus::Unknown:		return TEXT("Unknown");
+	}
+}
+
+void LexFromString(EUserPresenceStatus& OutStatus, const TCHAR* InStr)
+{
+	if (FCString::Stricmp(InStr, TEXT("Offline")) == 0)
+	{
+		OutStatus = EUserPresenceStatus::Offline;
+	}
+	else if (FCString::Stricmp(InStr, TEXT("Online")) == 0)
+	{
+		OutStatus = EUserPresenceStatus::Online;
+	}
+	else if (FCString::Stricmp(InStr, TEXT("Away")) == 0)
+	{
+		OutStatus = EUserPresenceStatus::Away;
+	}
+	else if (FCString::Stricmp(InStr, TEXT("ExtendedAway")) == 0)
+	{
+		OutStatus = EUserPresenceStatus::ExtendedAway;
+	}
+	else if (FCString::Stricmp(InStr, TEXT("DoNotDisturb")) == 0)
+	{
+		OutStatus = EUserPresenceStatus::DoNotDisturb;
+	}
+	else if (FCString::Stricmp(InStr, TEXT("Unknown")) == 0)
+	{
+		OutStatus = EUserPresenceStatus::Unknown;
+	}
+	else
+	{
+		checkNoEntry();
+		OutStatus = EUserPresenceStatus::Unknown;
+	}
+}
+
+const TCHAR* LexToString(EUserPresenceJoinability EnumVal)
+{
+	switch (EnumVal)
+	{
+	case EUserPresenceJoinability::Public:		return TEXT("Public");
+	case EUserPresenceJoinability::FriendsOnly:	return TEXT("FriendsOnly");
+	case EUserPresenceJoinability::InviteOnly:	return TEXT("InviteOnly");
+	case EUserPresenceJoinability::Private:		return TEXT("Private");
+	default:									checkNoEntry(); // Intentional fall-through
+	case EUserPresenceJoinability::Unknown:		return TEXT("Unknown");
+	}
+}
+
+void LexFromString(EUserPresenceJoinability& OutJoinability, const TCHAR* InStr)
+{
+	if (FCString::Stricmp(InStr, TEXT("Public")) == 0)
+	{
+		OutJoinability = EUserPresenceJoinability::Public;
+	}
+	else if (FCString::Stricmp(InStr, TEXT("FriendsOnly")) == 0)
+	{
+		OutJoinability = EUserPresenceJoinability::FriendsOnly;
+	}
+	else if (FCString::Stricmp(InStr, TEXT("InviteOnly")) == 0)
+	{
+		OutJoinability = EUserPresenceJoinability::InviteOnly;
+	}
+	else if (FCString::Stricmp(InStr, TEXT("Private")) == 0)
+	{
+		OutJoinability = EUserPresenceJoinability::Private;
+	}
+	else if (FCString::Stricmp(InStr, TEXT("Unknown")) == 0)
+	{
+		OutJoinability = EUserPresenceJoinability::Unknown;
+	}
+	else
+	{
+		checkNoEntry();
+		OutJoinability = EUserPresenceJoinability::Unknown;
+	}
+}
+
+const TCHAR* LexToString(EUserPresenceGameStatus EnumVal)
+{
+	switch (EnumVal)
+	{
+	case EUserPresenceGameStatus::PlayingThisGame:	return TEXT("PlayingThisGame");
+	case EUserPresenceGameStatus::PlayingOtherGame:	return TEXT("PlayingOtherGame");
+	default:										checkNoEntry(); // Intentional fall-through
+	case EUserPresenceGameStatus::Unknown:			return TEXT("Unknown");
+	}
+}
+
+void LexFromString(EUserPresenceGameStatus& OutGameStatus, const TCHAR* InStr)
+{
+	if (FCString::Stricmp(InStr, TEXT("PlayingThisGame")) == 0)
+	{
+		OutGameStatus = EUserPresenceGameStatus::PlayingThisGame;
+	}
+	else if (FCString::Stricmp(InStr, TEXT("PlayingOtherGame")) == 0)
+	{
+		OutGameStatus = EUserPresenceGameStatus::PlayingOtherGame;
+	}
+	else if (FCString::Stricmp(InStr, TEXT("Unknown")) == 0)
+	{
+		OutGameStatus = EUserPresenceGameStatus::Unknown;
+	}
+	else
+	{
+		checkNoEntry();
+		OutGameStatus = EUserPresenceGameStatus::Unknown;
+	}
+}
+
+/* UE::Online */ }
diff --git a/Engine/Plugins/Online/Experimental/OnlineServices/Source/OnlineServicesInterface/Public/Online/Presence.h b/Engine/Plugins/Online/Experimental/OnlineServices/Source/OnlineServicesInterface/Public/Online/Presence.h
--- a/Engine/Plugins/Online/Experimental/OnlineServices/Source/OnlineServicesInterface/Public/Online/Presence.h
+++ b/Engine/Plugins/Online/Experimental/OnlineServices/Source/OnlineServicesInterface/Public/Online/Presence.h
@@ -42,6 +42,8 @@ enum class EUserPresenceStatus : uint8
 	/** Default */
 	Unknown
 };
+ONLINESERVICESINTERFACE_API const TCHAR* LexToString(EUserPresenceStatus EnumVal);
+ONLINESERVICESINTERFACE_API void LexFromString(EUserPresenceStatus& OutStatus, const TCHAR* InStr);
 
 enum class EUserPresenceJoinability : uint8
 {
@@ -56,6 +58,8 @@ enum class EUserPresenceJoinability : uint8
 	/** Default */
 	Unknown
 };
+ONLINESERVICESINTERFACE_API const TCHAR* LexToString(EUserPresenceJoinability EnumVal);
+ONLINESERVICESINTERFACE_API void LexFromString(EUserPresenceJoinability& OutJoinability, const TCHAR* InStr);
 
 enum class EUserPresenceGameStatus : uint8
 {
@@ -66,6 +70,8 @@ enum class EUserPresenceGameStatus : uint8
 	/** Default */
 	Unknown
 };
+ONLINESERVICESINTERFACE_API const TCHAR* LexToString(EUserPresenceGameStatus EnumVal);
+ONLINESERVICESINTERFACE_API void LexFromString(EUserPresenceGameStatus& OutGameStatus, const TCHAR* InStr);
 
 typedef TMap<FString, FPresenceVariant> FPresenceProperties;
 
@@ -407,6 +413,8 @@ BEGIN_ONLINE_STRUCT_META(FPartialUpdatePresence::Params::FMutations)
 	ONLINE_STRUCT_FIELD(FPartialUpdatePresence::Params::FMutations, RemovedProperties)
 END_ONLINE_STRUCT_META()
 
+BEGIN_ONLINE_STRUCT_META(FPartialUpdatePresence::Result)
+END_ONLINE_STRUCT_META()
 
 BEGIN_ONLINE_STRUCT_META(FPresenceUpdated)
 	ONLINE_STRUCT_FIELD(FPresenceUpdated, LocalUserId),
