Updated OnlineServicesCommonEngineUtils to use unity files.

Before:
Total CPU Time: 19.1875 s
Total time in Parallel executor: 8.52 seconds

After:
Total CPU Time: 15.9375 s
Total time in Parallel executor: 8.28 seconds

diff --git a/Engine/Plugins/Online/Experimental/OnlineServices/Source/OnlineServicesCommonEngineUtils/OnlineServicesCommonEngineUtils.Build.cs b/Engine/Plugins/Online/Experimental/OnlineServices/Source/OnlineServicesCommonEngineUtils/OnlineServicesCommonEngineUtils.Build.cs
--- a/Engine/Plugins/Online/Experimental/OnlineServices/Source/OnlineServicesCommonEngineUtils/OnlineServicesCommonEngineUtils.Build.cs
+++ b/Engine/Plugins/Online/Experimental/OnlineServices/Source/OnlineServicesCommonEngineUtils/OnlineServicesCommonEngineUtils.Build.cs
@@ -8,8 +8,6 @@ public class OnlineServicesCommonEngineUtils : ModuleRules
 	{
 		PCHUsage = ModuleRules.PCHUsageMode.UseExplicitOrSharedPCHs;
 
-		bUseUnity = false; // set false to catch non-unity compile issues while prototyping
-
 		if (Target.bCompileAgainstEngine)
 		{
 			PublicDependencyModuleNames.AddRange(
