Implement OSSv2 schema attribute serialization and tests.

diff --git a/Engine/Plugins/Online/Experimental/OnlineServices/Source/OnlineServicesInterface/Private/Online/Schema.cpp b/Engine/Plugins/Online/Experimental/OnlineServices/Source/OnlineServicesInterface/Private/Online/Schema.cpp
--- a/Engine/Plugins/Online/Experimental/OnlineServices/Source/OnlineServicesInterface/Private/Online/Schema.cpp
+++ b/Engine/Plugins/Online/Experimental/OnlineServices/Source/OnlineServicesInterface/Private/Online/Schema.cpp
@@ -3,6 +3,9 @@
 #include "Online/Schema.h"
 #include "Algo/Find.h"
 #include "Algo/ForEach.h"
+#include "Algo/Reverse.h"
+#include "Algo/Transform.h"
+#include "Online/OnlineErrorDefinitions.h"
 
 DEFINE_LOG_CATEGORY_STATIC(LogOnlineSchema, Log, All);
 
@@ -16,6 +19,35 @@ T BuildFlagsFromArray(const TArray<T>& Input)
 	return OutValue;
 }
 
+ESchemaServiceAttributeSupportedTypeFlags TranslateAttributeType(ESchemaAttributeType InType)
+{
+	switch (InType)
+	{
+	case ESchemaAttributeType::Bool:
+		return ESchemaServiceAttributeSupportedTypeFlags::Bool;
+	case ESchemaAttributeType::Int64:
+		return ESchemaServiceAttributeSupportedTypeFlags::Int64;
+	case ESchemaAttributeType::Double:
+		return ESchemaServiceAttributeSupportedTypeFlags::Double;
+	case ESchemaAttributeType::String:
+		return ESchemaServiceAttributeSupportedTypeFlags::String;
+	default:
+		checkNoEntry();
+		// Intentional fall-through.
+	case ESchemaAttributeType::None:
+		return ESchemaServiceAttributeSupportedTypeFlags::None;
+	}
+}
+
+ESchemaServiceAttributeFlags TranslateAttributeFlags(ESchemaAttributeFlags InFlags)
+{
+	ESchemaServiceAttributeFlags OutFlags = ESchemaServiceAttributeFlags::None;
+	OutFlags |= EnumHasAnyFlags(InFlags, ESchemaAttributeFlags::Public) ? ESchemaServiceAttributeFlags::Public : ESchemaServiceAttributeFlags::None;
+	OutFlags |= EnumHasAnyFlags(InFlags, ESchemaAttributeFlags::Private) ? ESchemaServiceAttributeFlags::Private : ESchemaServiceAttributeFlags::None;
+	OutFlags |= EnumHasAnyFlags(InFlags, ESchemaAttributeFlags::Searchable) ? ESchemaServiceAttributeFlags::Searchable : ESchemaServiceAttributeFlags::None;
+	return OutFlags;
+}
+
 bool FSchemaRegistry::ParseConfig(const FSchemaRegistryDescriptorConfig& Config)
 {
 	bool ParsedSuccessfully = true;
@@ -63,7 +95,7 @@ bool FSchemaRegistry::ParseConfig(const FSchemaRegistryDescriptorConfig& Config)
 			}
 			SeenAttributes.Add(SchemaServiceAttributeDescriptor.Id);
 
-			ESchemaServiceAttributeSupportedTypeFlags SupportedTypes = BuildFlagsFromArray(SchemaServiceAttributeDescriptor.SupportedTypes);
+			const ESchemaServiceAttributeSupportedTypeFlags SupportedTypes = BuildFlagsFromArray(SchemaServiceAttributeDescriptor.SupportedTypes);
 
 			// Check that supported type has been set.
 			if (SupportedTypes == ESchemaServiceAttributeSupportedTypeFlags::None)
@@ -82,6 +114,16 @@ bool FSchemaRegistry::ParseConfig(const FSchemaRegistryDescriptorConfig& Config)
 					*SchemaServiceAttributeDescriptor.Id.ToString().ToLower());
 				ParsedSuccessfully = false;
 			}
+
+			// Check for valid attribute flags
+			const ESchemaServiceAttributeFlags ServiceAttributeFlags = BuildFlagsFromArray(SchemaServiceAttributeDescriptor.Flags);
+			if (!EnumHasAnyFlags(ServiceAttributeFlags, ESchemaServiceAttributeFlags::Public | ESchemaServiceAttributeFlags::Private))
+			{
+				UE_LOG(LogOnlineSchema, Error, TEXT("Invalid schema service attribute %s.%s: Either Public or Private visibility must be set."),
+					*SchemaServiceDescriptor.Id.ToString().ToLower(),
+					*SchemaServiceAttributeDescriptor.Id.ToString().ToLower());
+				ParsedSuccessfully = false;
+			}
 		}
 
 		KnownSchemaServiceDescriptors.Add(SchemaServiceDescriptor.Id, &SchemaServiceDescriptor);
@@ -224,6 +266,16 @@ bool FSchemaRegistry::ParseConfig(const FSchemaRegistryDescriptorConfig& Config)
 				}
 				SeenAttributes.Add(SchemaAttributeDescriptor.Id);
 
+				// Check that a type has been set
+				if (SchemaAttributeDescriptor.Type == ESchemaAttributeType::None)
+				{
+					UE_LOG(LogOnlineSchema, Error, TEXT("Invalid schema attribute %s.%s.%s: A valid type must be set."),
+						*SchemaDescriptor.Id.ToString().ToLower(),
+						*SchemaCategoryDescriptor.Id.ToString().ToLower(),
+						*SchemaAttributeDescriptor.Id.ToString().ToLower());
+					ParsedSuccessfully = false;
+				}
+
 				// Check that variable length data has a set maximum length.
 				if (SchemaAttributeDescriptor.Type == ESchemaAttributeType::String && SchemaAttributeDescriptor.MaxSize <= 0)
 				{
@@ -234,7 +286,7 @@ bool FSchemaRegistry::ParseConfig(const FSchemaRegistryDescriptorConfig& Config)
 					ParsedSuccessfully = false;
 				}
 
-				ESchemaAttributeFlags AttributeFlags = BuildFlagsFromArray(SchemaAttributeDescriptor.Flags);
+				const ESchemaAttributeFlags AttributeFlags = BuildFlagsFromArray(SchemaAttributeDescriptor.Flags);
 
 				// Check that searchable attributes only exist in a base schema.
 				if (EnumHasAllFlags(AttributeFlags, ESchemaAttributeFlags::Searchable) && SchemaDescriptor.ParentId != FSchemaId())
@@ -246,6 +298,32 @@ bool FSchemaRegistry::ParseConfig(const FSchemaRegistryDescriptorConfig& Config)
 					ParsedSuccessfully = false;
 				}
 
+				// Check that public or private visibility has been set and is valid.
+				if (EnumHasAllFlags(AttributeFlags, ESchemaAttributeFlags::Public | ESchemaAttributeFlags::Private))
+				{
+					UE_LOG(LogOnlineSchema, Error, TEXT("Invalid schema attribute %s.%s.%s: Public and Private visibility are mutually exclusive."),
+						*SchemaDescriptor.Id.ToString().ToLower(),
+						*SchemaCategoryDescriptor.Id.ToString().ToLower(),
+						*SchemaAttributeDescriptor.Id.ToString().ToLower());
+					ParsedSuccessfully = false;
+				}
+				if (!EnumHasAnyFlags(AttributeFlags, ESchemaAttributeFlags::Public | ESchemaAttributeFlags::Private))
+				{
+					UE_LOG(LogOnlineSchema, Error, TEXT("Invalid schema attribute %s.%s.%s: Either Public or Private visibility must be set."),
+						*SchemaDescriptor.Id.ToString().ToLower(),
+						*SchemaCategoryDescriptor.Id.ToString().ToLower(),
+						*SchemaAttributeDescriptor.Id.ToString().ToLower());
+					ParsedSuccessfully = false;
+				}
+				if (EnumHasAllFlags(AttributeFlags, ESchemaAttributeFlags::Searchable | ESchemaAttributeFlags::Private))
+				{
+					UE_LOG(LogOnlineSchema, Error, TEXT("Invalid schema attribute %s.%s.%s: Searchable attributes may not be private."),
+						*SchemaDescriptor.Id.ToString().ToLower(),
+						*SchemaCategoryDescriptor.Id.ToString().ToLower(),
+						*SchemaAttributeDescriptor.Id.ToString().ToLower());
+					ParsedSuccessfully = false;
+				}
+
 				// Check that schema compatibility attribute only exists in a base schema and used only once and that it is the Int64 type.
 				if (EnumHasAllFlags(AttributeFlags, ESchemaAttributeFlags::SchemaCompatibilityId))
 				{
@@ -280,6 +358,15 @@ bool FSchemaRegistry::ParseConfig(const FSchemaRegistryDescriptorConfig& Config)
 							*SchemaAttributeDescriptor.Id.ToString().ToLower());
 						ParsedSuccessfully = false;
 					}
+
+					if (EnumHasAnyFlags(AttributeFlags, ESchemaAttributeFlags::Private))
+					{
+						UE_LOG(LogOnlineSchema, Error, TEXT("Invalid schema attribute %s.%s.%s: SchemaCompatibilityId attribute may not be private."),
+							*SchemaDescriptor.Id.ToString().ToLower(),
+							*SchemaCategoryDescriptor.Id.ToString().ToLower(),
+							*SchemaAttributeDescriptor.Id.ToString().ToLower());
+						ParsedSuccessfully = false;
+					}
 				}
 
 				// Check that group id is valid.
@@ -297,6 +384,8 @@ bool FSchemaRegistry::ParseConfig(const FSchemaRegistryDescriptorConfig& Config)
 
 	// Build definitions only when there are no parsing errors.
 	TMap<FSchemaId, TSharedRef<const FSchemaDefinition>> ParsedSchemaDefinitions;
+	TMap<int64, TSharedRef<const FSchemaDefinition>> ParsedSchemaDefinitionsByCompatibilityId;
+
 	if (ParsedSuccessfully)
 	{
 		// Build schema definitions
@@ -311,8 +400,9 @@ bool FSchemaRegistry::ParseConfig(const FSchemaRegistryDescriptorConfig& Config)
 			SchemaDefinition->Id = SchemaDescriptor.Id;
 
 			TMap<FSchemaCategoryId, FSchemaServiceDescriptorId> ServiceDescriptors;
+			TArray<FSchemaAttributeDefinition*> SchemaAttributeDefinitionAssignmentStorage;
 
-			// Populate definition from descriptor
+			// Populate definition from descriptor.
 			for (const FSchemaDescriptor** IterSchemaDescriptor = KnownSchemaDescriptors.Find(SchemaDescriptor.Id);
 				IterSchemaDescriptor != nullptr;
 				IterSchemaDescriptor = KnownSchemaDescriptors.Find((*IterSchemaDescriptor)->ParentId))
@@ -330,6 +420,7 @@ bool FSchemaRegistry::ParseConfig(const FSchemaRegistryDescriptorConfig& Config)
 				for (const FSchemaCategoryDescriptor& CurrentSchemaCategoryDescriptor : CurrentSchemaDescriptor.Categories)
 				{
 					FSchemaCategoryDefinition& SchemaCategoryDefinition = SchemaDefinition->Categories.FindOrAdd(CurrentSchemaCategoryDescriptor.Id);
+					SchemaCategoryDefinition.Id = CurrentSchemaCategoryDescriptor.Id;
 					SchemaDataCrc = FCrc::TypeCrc32(CurrentSchemaCategoryDescriptor.Id, SchemaDataCrc);
 
 					for (const FSchemaAttributeDescriptor& CurrentSchemaAttributeDescriptor : CurrentSchemaCategoryDescriptor.Attributes)
@@ -339,10 +430,18 @@ bool FSchemaRegistry::ParseConfig(const FSchemaRegistryDescriptorConfig& Config)
 						AttributeDefinition.Flags = BuildFlagsFromArray(CurrentSchemaAttributeDescriptor.Flags);
 						AttributeDefinition.Type = CurrentSchemaAttributeDescriptor.Type;
 						AttributeDefinition.MaxSize = CurrentSchemaAttributeDescriptor.MaxSize;
+
+						if (EnumHasAnyFlags(AttributeDefinition.Flags, ESchemaAttributeFlags::SchemaCompatibilityId))
+						{
+							SchemaCategoryDefinition.SchemaCompatibilityAttributeId = AttributeDefinition.Id;
+						}
+
 						SchemaDataCrc = FCrc::TypeCrc32(AttributeDefinition.Id, SchemaDataCrc);
 						SchemaDataCrc = FCrc::TypeCrc32(AttributeDefinition.Flags, SchemaDataCrc);
 						SchemaDataCrc = FCrc::TypeCrc32(AttributeDefinition.Type, SchemaDataCrc);
 						SchemaDataCrc = FCrc::TypeCrc32(AttributeDefinition.MaxSize, SchemaDataCrc);
+
+						SchemaAttributeDefinitionAssignmentStorage.Add(&AttributeDefinition);
 					}
 
 					// Store service descriptor for assignment.
@@ -353,6 +452,9 @@ bool FSchemaRegistry::ParseConfig(const FSchemaRegistryDescriptorConfig& Config)
 				}
 			}
 
+			// Attributes must be ordered from parent to derived schema to ensure that service attribute assignment is consistent between derived and parent schemas.
+			Algo::Reverse(SchemaAttributeDefinitionAssignmentStorage);
+
 			// Assign attributes to service attributes.
 			for (TPair<FSchemaCategoryId, FSchemaServiceDescriptorId>& ServiceDescriptor : ServiceDescriptors)
 			{
@@ -362,15 +464,109 @@ bool FSchemaRegistry::ParseConfig(const FSchemaRegistryDescriptorConfig& Config)
 				// todo: pack multiple attributes into a service attribute.
 				// For the time being schema attributes and service attributes have a 1:1 relationship.
 
-				// todo
+				// Service attribute map.
+				TMap<FSchemaServiceAttributeId, const FSchemaServiceAttributeDescriptor*> ServiceAttributesById;
+
+				// Bucket service attributes for assignment
+				TMap<const FSchemaServiceAttributeDescriptor*, ESchemaServiceAttributeSupportedTypeFlags> PublicAttributes;
+				TMap<const FSchemaServiceAttributeDescriptor*, ESchemaServiceAttributeSupportedTypeFlags> PrivateAttributes;
+				TMap<const FSchemaServiceAttributeDescriptor*, ESchemaServiceAttributeSupportedTypeFlags> SearchableAttributes;
+
+				// Build out buckets.
+				for (const FSchemaServiceAttributeDescriptor& SchemaServiceAttributeDescriptor : SchemaServiceDescriptor.Attributes)
+				{
+					ServiceAttributesById.Add(SchemaServiceAttributeDescriptor.Id, &SchemaServiceAttributeDescriptor);
+
+					const ESchemaServiceAttributeSupportedTypeFlags SupportedTypes = BuildFlagsFromArray(SchemaServiceAttributeDescriptor.SupportedTypes);
+					const ESchemaServiceAttributeFlags ServiceAttributeFlags = BuildFlagsFromArray(SchemaServiceAttributeDescriptor.Flags);
+					if (EnumHasAnyFlags(ServiceAttributeFlags, ESchemaServiceAttributeFlags::Public))
+					{
+						PublicAttributes.Add(&SchemaServiceAttributeDescriptor, SupportedTypes);
+					}
+					if (EnumHasAnyFlags(ServiceAttributeFlags, ESchemaServiceAttributeFlags::Private))
+					{
+						PrivateAttributes.Add(&SchemaServiceAttributeDescriptor, SupportedTypes);
+					}
+					if (EnumHasAnyFlags(ServiceAttributeFlags, ESchemaServiceAttributeFlags::Searchable))
+					{
+						SearchableAttributes.Add(&SchemaServiceAttributeDescriptor, SupportedTypes);
+					}
+				}
+
+				// Assign schema attributes based on buckets.
+				for (FSchemaAttributeDefinition* AttributeDefinition : SchemaAttributeDefinitionAssignmentStorage)
+				{
+					// Find a service attribute which can handle the schema attribute.
+					TMap<const FSchemaServiceAttributeDescriptor*, ESchemaServiceAttributeSupportedTypeFlags>* Bucket = nullptr;
+					if (EnumHasAnyFlags(AttributeDefinition->Flags, ESchemaAttributeFlags::Public))
+					{
+						Bucket = EnumHasAnyFlags(AttributeDefinition->Flags, ESchemaAttributeFlags::Searchable) ? &SearchableAttributes : &PublicAttributes;
+					}
+					else
+					{
+						Bucket = &PrivateAttributes;
+					}
+
+					const TPair<const FSchemaServiceAttributeDescriptor*, ESchemaServiceAttributeSupportedTypeFlags>* SchemaServiceAttributeDescriptorPair = Algo::FindByPredicate(*Bucket,
+					[AttributeDefinition](const TPair<const FSchemaServiceAttributeDescriptor*, ESchemaServiceAttributeSupportedTypeFlags>& ServiceAttributePair)
+					{
+						const FSchemaServiceAttributeDescriptor* SchemaServiceAttributeDescriptor = ServiceAttributePair.Get<0>();
+						const ESchemaServiceAttributeSupportedTypeFlags SchemaServiceAttributeSupportedTypeFlags = ServiceAttributePair.Get<1>();
+						return EnumHasAnyFlags(SchemaServiceAttributeSupportedTypeFlags, TranslateAttributeType(AttributeDefinition->Type)) &&
+							AttributeDefinition->MaxSize <= SchemaServiceAttributeDescriptor->MaxSize;
+					});
+					const FSchemaServiceAttributeDescriptor* FoundServiceAttributeDescriptor = SchemaServiceAttributeDescriptorPair ? SchemaServiceAttributeDescriptorPair->Get<0>() : nullptr;
+
+					if (FoundServiceAttributeDescriptor)
+					{
+						// Todo: Reuse descriptor for assignment until it is full.
+
+						// Remove found descriptor from buckets.
+						PublicAttributes.Remove(FoundServiceAttributeDescriptor);
+						PrivateAttributes.Remove(FoundServiceAttributeDescriptor);
+						SearchableAttributes.Remove(FoundServiceAttributeDescriptor);
+
+						// Create service attribute definition.
+						FSchemaServiceAttributeDefinition& SchemaServiceAttributeDefinition = SchemaCategoryDefinition.ServiceAttributeDefinitions.Add(FoundServiceAttributeDescriptor->Id);
+						SchemaServiceAttributeDefinition.Id = FoundServiceAttributeDescriptor->Id;
+						SchemaServiceAttributeDefinition.Type = TranslateAttributeType(AttributeDefinition->Type);
+						SchemaServiceAttributeDefinition.Flags = TranslateAttributeFlags(AttributeDefinition->Flags);
+						SchemaServiceAttributeDefinition.MaxSize = FoundServiceAttributeDescriptor->MaxSize;
+						SchemaServiceAttributeDefinition.SchemaAttributeIds.Add(AttributeDefinition->Id);
+						AttributeDefinition->ServiceAttributeId = SchemaServiceAttributeDefinition.Id;
+
+						if (EnumHasAnyFlags(AttributeDefinition->Flags, ESchemaAttributeFlags::SchemaCompatibilityId))
+						{
+							SchemaCategoryDefinition.SchemaCompatibilityServiceAttributeId = SchemaServiceAttributeDefinition.Id;
+						}
+
+						// Add service attribute definition to data crc.
+						SchemaDataCrc = FCrc::TypeCrc32(SchemaServiceAttributeDefinition.Id, SchemaDataCrc);
+						SchemaDataCrc = FCrc::TypeCrc32(SchemaServiceAttributeDefinition.Type, SchemaDataCrc);
+						SchemaDataCrc = FCrc::TypeCrc32(SchemaServiceAttributeDefinition.Flags, SchemaDataCrc);
+						SchemaDataCrc = FCrc::TypeCrc32(SchemaServiceAttributeDefinition.MaxSize, SchemaDataCrc);
+						Algo::ForEach(SchemaServiceAttributeDefinition.SchemaAttributeIds, [&SchemaDataCrc](const FSchemaAttributeId& Id){ SchemaDataCrc = FCrc::TypeCrc32(Id, SchemaDataCrc); });
+					}
+					else
+					{
+						UE_LOG(LogOnlineSchema, Error, TEXT("Failed to find service attribute to fit schema attribute %s.%s.%s."),
+							*SchemaDefinition->Id.ToString().ToLower(),
+							*SchemaCategoryDefinition.Id.ToString().ToLower(),
+							*AttributeDefinition->Id.ToString().ToLower());
+						ParsedSuccessfully = false;
+					}
+				}
+
+				SchemaDataCrc = FCrc::TypeCrc32(SchemaCategoryDefinition.SchemaCompatibilityAttributeId, SchemaDataCrc);
+				SchemaDataCrc = FCrc::TypeCrc32(SchemaCategoryDefinition.SchemaCompatibilityServiceAttributeId, SchemaDataCrc);
 			}
 
-			// Build compatibility id.
+			// Build compatibility id from the schema and data CRCs.
 			{
 				const uint32 SchemaIdCrc = FCrc::TypeCrc32(SchemaDefinition->Id);
 				if (SeenSchemaCrcs.Find(SchemaIdCrc) != nullptr)
 				{
-					UE_LOG(LogOnlineSchema, Error, TEXT("Invalid schema %s: Hash collision processing schema name, please rename."),
+					UE_LOG(LogOnlineSchema, Error, TEXT("Invalid schema %s: CRC collision processing schema name, please rename the schema to avoid collision."),
 						*SchemaDefinition->Id.ToString().ToLower());
 					ParsedSuccessfully = false;
 				}
@@ -379,12 +575,15 @@ bool FSchemaRegistry::ParseConfig(const FSchemaRegistryDescriptorConfig& Config)
 				const uint64 CompatibilityIdUint = (static_cast<uint64>(SchemaIdCrc) << 32) | SchemaDataCrc;
 				SchemaDefinition->CompatibilityId = *reinterpret_cast<const int64*>(&CompatibilityIdUint);
 			}
+
+			ParsedSchemaDefinitionsByCompatibilityId.Add(SchemaDefinition->CompatibilityId, SchemaDefinition);
 		}
 	}
 
 	if (ParsedSuccessfully)
 	{
-		SchemaDefinitions = MoveTemp(ParsedSchemaDefinitions);
+		SchemaDefinitionsById = MoveTemp(ParsedSchemaDefinitions);
+		SchemaDefinitionsByCompatibilityId = MoveTemp(ParsedSchemaDefinitionsByCompatibilityId);
 		return true;
 	}
 	else
@@ -394,72 +593,382 @@ bool FSchemaRegistry::ParseConfig(const FSchemaRegistryDescriptorConfig& Config)
 }
 
 FSchemaCategoryInstance::FSchemaCategoryInstance(
-	const FSchemaId& SchemaId,
+	const FSchemaId& DerivedSchemaId,
 	const FSchemaId& BaseSchemaId,
 	const FSchemaCategoryId& CategoryId,
-	const TSharedRef<const FSchemaRegistry>& Registry)
+	const TSharedRef<const FSchemaRegistry>& SchemaRegistry)
+	: SchemaRegistry(SchemaRegistry)
 {
-	if ((SchemaDefinition = Registry->GetDefinition(SchemaId)))
+	bool bInitSuccess = true;
+
+	if (InitializeSchemaDefinition(BaseSchemaId, CategoryId, &BaseSchemaDefinition, &BaseSchemaCategoryDefinition))
 	{
-		if (!Registry->IsSchemaChildOf(SchemaId, BaseSchemaId))
+		if (DerivedSchemaId == FSchemaId())
 		{
-			UE_LOG(LogOnlineSchema, Error, TEXT("SchemaCategoryInstance init error: Schema %s is not a child of %s"),
-				*SchemaId.ToString().ToLower(),
-				*BaseSchemaId.ToString().ToLower());
+			// Derived schema id is unknown. One of the below is true:
+			// 1. All attributes are defined in the base schema.
+			// 2. An attribute in the base schema is flagged as SchemaCompatibilityId to allow
+			//    detection of the derived schema based on incoming service attribute data.
 		}
 		else
 		{
-			SchemaCategoryDefinition = SchemaDefinition->Categories.Find(CategoryId);
-			if (SchemaCategoryDefinition == nullptr)
+			if (InitializeSchemaDefinition(DerivedSchemaId, CategoryId, &DerivedSchemaDefinition, &DerivedSchemaCategoryDefinition))
 			{
-				UE_LOG(LogOnlineSchema, Error, TEXT("SchemaCategoryInstance init error: Unable to find category %s in schema %s"),
-					*SchemaId.ToString().ToLower());
+				// When a derived schema is in use the base schema must have a field flagged as SchemaCompatibilityId.
+				if (BaseSchemaCategoryDefinition->SchemaCompatibilityServiceAttributeId == FSchemaServiceAttributeId())
+				{
+					UE_LOG(LogOnlineSchema, Error, TEXT("[FSchemaCategoryInstance] init error: Base schema %s is missing a valid SchemaCompatibilityId field"),
+						*BaseSchemaId.ToString().ToLower());
+					bInitSuccess = false;
+				}
+				else if (!SchemaRegistry->IsSchemaChildOf(DerivedSchemaId, BaseSchemaId))
+				{
+					UE_LOG(LogOnlineSchema, Error, TEXT("[FSchemaCategoryInstance] init error: Schema %s is not a child of %s"),
+						*DerivedSchemaId.ToString().ToLower(),
+						*BaseSchemaId.ToString().ToLower());
+					bInitSuccess = false;
+				}
+			}
+			else
+			{
+				// Unable to find derived schema when it was expected to be valid.
+				bInitSuccess = false;
 			}
 		}
 	}
-	else
+
+	if (!bInitSuccess)
 	{
-		UE_LOG(LogOnlineSchema, Error, TEXT("SchemaCategoryInstance init error: Unable to find definition for schema %s"),
-			*SchemaId.ToString().ToLower());
+		// Invalidate schema to fail initialization.
+		BaseSchemaDefinition.Reset();
+		BaseSchemaCategoryDefinition = nullptr;
+		DerivedSchemaDefinition.Reset();
+		DerivedSchemaCategoryDefinition = nullptr;
 	}
 }
 
-TOnlineResult<FSchemaCategoryInstanceApplyApplicationChanges> FSchemaCategoryInstance::ApplyChanges(FSchemaCategoryInstanceApplyApplicationChanges::Params&& Changes)
+TOnlineResult<FSchemaCategoryInstanceApplyApplicationDelta> FSchemaCategoryInstance::ApplyApplicationDelta(FSchemaCategoryInstanceApplyApplicationDelta::Params&& Params)
 {
 	// todo: handle packing multiple attributes into a service attribute.
 
-	// todo
-	return TOnlineResult<FSchemaCategoryInstanceApplyApplicationChanges>(FSchemaCategoryInstanceApplyApplicationChanges::Result({}));
+	const FSchemaDefinition* SchemaDefinition = nullptr;
+	const FSchemaCategoryDefinition* SchemaCategoryDefinition = nullptr;
+	if (!GetSerializationSchema(&SchemaDefinition, &SchemaCategoryDefinition))
+	{
+		UE_LOG(LogOnlineSchema, Warning, TEXT("[ApplyApplicationDelta] Schema category definition not found."));
+		return TOnlineResult<FSchemaCategoryInstanceApplyApplicationDelta>(Errors::InvalidState());
+	}
+
+	// Configured service attributes must be present to serialize to service.
+	if (SchemaCategoryDefinition->ServiceAttributeDefinitions.IsEmpty())
+	{
+		UE_LOG(LogOnlineSchema, Warning, TEXT("[ApplyApplicationDelta] %s.%s: Schema category has not been configured for serialization."),
+			*SchemaDefinition->Id.ToString().ToLower(),
+			*SchemaCategoryDefinition->Id.ToString().ToLower());
+		return TOnlineResult<FSchemaCategoryInstanceApplyApplicationDelta>(Errors::InvalidParams());
+	}
+
+	FSchemaServiceChanges Changes;
+
+	// Translate mutated attributes.
+	for (TPair<FSchemaAttributeId, FSchemaVariant>& UpdatedAttributeData : Params.UpdatedAttributes)
+	{
+		const FSchemaAttributeId& UpdatedAttributeId = UpdatedAttributeData.Get<0>();
+		FSchemaVariant& UpdatedAttributeValue = UpdatedAttributeData.Get<1>();
+
+		if (const FSchemaAttributeDefinition* AttributeDefinition = SchemaCategoryDefinition->SchemaAttributeDefinitions.Find(UpdatedAttributeId))
+		{
+			// Schema compatibility id cannot be modified by application code.
+			if (EnumHasAnyFlags(AttributeDefinition->Flags, ESchemaAttributeFlags::SchemaCompatibilityId))
+			{
+				UE_LOG(LogOnlineSchema, Warning, TEXT("[ApplyApplicationDelta] %s.%s.%s: It is not valid to set SchemaCompatibilityId in an attribute update"),
+					*SchemaDefinition->Id.ToString().ToLower(),
+					*SchemaCategoryDefinition->Id.ToString().ToLower(),
+					*AttributeDefinition->Id.ToString().ToLower());
+				return TOnlineResult<FSchemaCategoryInstanceApplyApplicationDelta>(Errors::InvalidParams());
+			}
+
+			// Check that type matches.
+			if (AttributeDefinition->Type != UpdatedAttributeValue.GetType())
+			{
+				UE_LOG(LogOnlineSchema, Warning, TEXT("[ApplyApplicationDelta] %s.%s.%s: Attribute type %s does not match the expected type %s."),
+					*SchemaDefinition->Id.ToString().ToLower(),
+					*SchemaCategoryDefinition->Id.ToString().ToLower(),
+					*AttributeDefinition->Id.ToString().ToLower(),
+					LexToString(UpdatedAttributeValue.GetType()),
+					LexToString(AttributeDefinition->Type));
+				return TOnlineResult<FSchemaCategoryInstanceApplyApplicationDelta>(Errors::InvalidParams());
+			}
+
+			// Check that variable sized data fits within the service attribute.
+			if (AttributeDefinition->Type == ESchemaAttributeType::String && UpdatedAttributeValue.GetString().Len() > AttributeDefinition->MaxSize)
+			{
+				UE_LOG(LogOnlineSchema, Warning, TEXT("[ApplyApplicationDelta] %s.%s.%s: Variably sized attribute exceeds maximum defined length. %d > %d."),
+					*SchemaDefinition->Id.ToString().ToLower(),
+					*SchemaCategoryDefinition->Id.ToString().ToLower(),
+					*AttributeDefinition->Id.ToString().ToLower(),
+					UpdatedAttributeValue.GetString().Len(),
+					AttributeDefinition->MaxSize);
+				return TOnlineResult<FSchemaCategoryInstanceApplyApplicationDelta>(Errors::InvalidParams());
+			}
+
+			Changes.UpdatedAttributes.Add(AttributeDefinition->ServiceAttributeId, {AttributeDefinition->ServiceAttributeId, AttributeDefinition->Flags, MoveTemp(UpdatedAttributeValue)});
+		}
+		else
+		{
+			UE_LOG(LogOnlineSchema, Warning, TEXT("[ApplyApplicationDelta] %s.%s.%s: Attribute definition does not exist in schema category."),
+				*SchemaDefinition->Id.ToString().ToLower(),
+				*SchemaCategoryDefinition->Id.ToString().ToLower(),
+				*UpdatedAttributeId.ToString().ToLower());
+			return TOnlineResult<FSchemaCategoryInstanceApplyApplicationDelta>(Errors::InvalidParams());
+		}
+	}
+
+	// Translate removed attributes.
+	for (FSchemaAttributeId& RemovedAttributeId : Params.RemovedAttributes)
+	{
+		if (const FSchemaAttributeDefinition* AttributeDefinition = SchemaCategoryDefinition->SchemaAttributeDefinitions.Find(RemovedAttributeId))
+		{
+			// Schema compatibility id cannot be modified by application code.
+			if (EnumHasAnyFlags(AttributeDefinition->Flags, ESchemaAttributeFlags::SchemaCompatibilityId))
+			{
+				UE_LOG(LogOnlineSchema, Warning, TEXT("[ApplyApplicationDelta] %s.%s.%s: It is not valid to set SchemaCompatibilityId in an attribute update"),
+					*SchemaDefinition->Id.ToString().ToLower(),
+					*SchemaCategoryDefinition->Id.ToString().ToLower(),
+					*AttributeDefinition->Id.ToString().ToLower());
+				return TOnlineResult<FSchemaCategoryInstanceApplyApplicationDelta>(Errors::InvalidParams());
+			}
+			else
+			{
+				Changes.RemovedAttributes.Add(AttributeDefinition->ServiceAttributeId);
+			}
+		}
+		else
+		{
+			UE_LOG(LogOnlineSchema, Warning, TEXT("[ApplyApplicationDelta] %s.%s.%s: Attribute definition does not exist in schema category."),
+				*SchemaDefinition->Id.ToString().ToLower(),
+				*SchemaCategoryDefinition->Id.ToString().ToLower(),
+				*RemovedAttributeId.ToString().ToLower());
+			return TOnlineResult<FSchemaCategoryInstanceApplyApplicationDelta>(Errors::InvalidParams());
+		}
+	}
+
+	// Check if SchemaCompatibilityId needs to be set.
+	if (SchemaCategoryDefinition->SchemaCompatibilityAttributeId != FSchemaAttributeId() && LastSentSchemaCompatibilityId == 0)
+	{
+		const FSchemaAttributeDefinition* AttributeDefinition = SchemaCategoryDefinition->SchemaAttributeDefinitions.Find(SchemaCategoryDefinition->SchemaCompatibilityAttributeId);
+		if (ensure(AttributeDefinition))
+		{
+			Changes.UpdatedAttributes.Add(AttributeDefinition->ServiceAttributeId, { AttributeDefinition->ServiceAttributeId, AttributeDefinition->Flags, SchemaDefinition->CompatibilityId });
+			LastSentSchemaCompatibilityId = SchemaDefinition->CompatibilityId;
+		}
+		else
+		{
+			UE_LOG(LogOnlineSchema, Warning, TEXT("[ApplyApplicationDelta] %s.%s.%s: Failed to find attribute defnition for SchemaCompatibilityId field."),
+				*SchemaDefinition->Id.ToString().ToLower(),
+				*SchemaCategoryDefinition->Id.ToString().ToLower(),
+				*SchemaCategoryDefinition->SchemaCompatibilityAttributeId.ToString().ToLower());
+			return TOnlineResult<FSchemaCategoryInstanceApplyApplicationDelta>(Errors::InvalidState());
+		}
+	}
+
+	return TOnlineResult<FSchemaCategoryInstanceApplyApplicationDelta>(FSchemaCategoryInstanceApplyApplicationDelta::Result({MoveTemp(Changes)}));
 }
 
-TOnlineResult<FSchemaCategoryInstanceApplyServiceChanges> FSchemaCategoryInstance::ApplyChanges(FSchemaCategoryInstanceApplyServiceChanges::Params&& Changes)
+TOnlineResult<FSchemaCategoryInstanceApplyServiceSnapshot> FSchemaCategoryInstance::ApplyServiceSnapshot(FSchemaCategoryInstanceApplyServiceSnapshot::Params&& Params)
 {
 	// todo: handle unpacking a service attribute into multiple attributes.
 
-	// todo
-	return TOnlineResult<FSchemaCategoryInstanceApplyServiceChanges>(FSchemaCategoryInstanceApplyServiceChanges::Result({}));
+	const FSchemaDefinition* SchemaDefinition = nullptr;
+	const FSchemaCategoryDefinition* SchemaCategoryDefinition = nullptr;
+
+	if (!GetSerializationSchema(&SchemaDefinition, &SchemaCategoryDefinition))
+	{
+		UE_LOG(LogOnlineSchema, Warning, TEXT("[ApplyServiceSnapshot] Schema category definition not found."));
+		return TOnlineResult<FSchemaCategoryInstanceApplyServiceSnapshot>(Errors::InvalidState());
+	}
+
+	FSchemaApplicationChanges Changes;
+
+	// Check for a schema change.
+	TPair<FSchemaServiceAttributeId, FSchemaVariant>* SchemaCompatibilityAttributeData = Algo::FindByPredicate(Params.Attributes,
+		[&SchemaCategoryDefinition](const TPair<FSchemaServiceAttributeId, FSchemaVariant>& AttributeData)
+		{
+			return AttributeData.Get<0>() == SchemaCategoryDefinition->SchemaCompatibilityServiceAttributeId;
+		});
+
+	// Handle incoming schema change.
+	if (SchemaCompatibilityAttributeData)
+	{
+		const FSchemaVariant& SchemaCompatibilityIdValue = SchemaCompatibilityAttributeData->Get<1>();
+
+		if (TSharedPtr<const FSchemaDefinition> SwappedSchema = SchemaRegistry->GetDefinition(SchemaCompatibilityIdValue.GetInt64()))
+		{
+			if (!SchemaRegistry->IsSchemaChildOf(SwappedSchema->Id, BaseSchemaDefinition->Id))
+			{
+				UE_LOG(LogOnlineSchema, Warning, TEXT("[ApplyServiceSnapshot] %s.%s.%s: Invalid schema swap. Incoming schema %s is not a child of %s."),
+					*SchemaDefinition->Id.ToString().ToLower(),
+					*SchemaCategoryDefinition->Id.ToString().ToLower(),
+					*SchemaCategoryDefinition->SchemaCompatibilityServiceAttributeId.ToString().ToLower(),
+					*SwappedSchema->Id.ToString().ToLower(),
+					*BaseSchemaDefinition->Id.ToString().ToLower());
+				return TOnlineResult<FSchemaCategoryInstanceApplyServiceSnapshot>(Errors::InvalidParams());
+			}
+
+			// Check that schema has actually changed.
+			if (DerivedSchemaDefinition == nullptr || DerivedSchemaDefinition->Id != SwappedSchema->Id)
+			{
+				Changes.SchemaId.Emplace(SwappedSchema->Id);
+
+				// Category must exist due to parent -> child relationship
+				const FSchemaCategoryDefinition* SwappedSchemaCategoryDefinition = SwappedSchema->Categories.Find(SchemaCategoryDefinition->Id);
+				if (ensure(SwappedSchemaCategoryDefinition))
+				{
+					UE_LOG(LogOnlineSchema, Verbose, TEXT("[ApplyServiceSnapshot] %s.%s.%s: Derived schema changed from %s to %s."),
+						*SchemaDefinition->Id.ToString().ToLower(),
+						*SchemaCategoryDefinition->Id.ToString().ToLower(),
+						*SchemaCategoryDefinition->SchemaCompatibilityServiceAttributeId.ToString().ToLower(),
+						DerivedSchemaDefinition ? *DerivedSchemaDefinition->Id.ToString().ToLower() : TEXT("nullptr"),
+						*SwappedSchemaCategoryDefinition->Id.ToString().ToLower());
+
+					// Set new schema to class member data.
+					DerivedSchemaDefinition = SwappedSchema;
+					DerivedSchemaCategoryDefinition = SwappedSchemaCategoryDefinition;
+
+					// Update schema definitions for current update processing.
+					SchemaDefinition = DerivedSchemaDefinition.Get();
+					SchemaCategoryDefinition = SwappedSchemaCategoryDefinition;
+				}
+				else
+				{
+					UE_LOG(LogOnlineSchema, Error, TEXT("[ApplyServiceSnapshot] %s.%s.%s: Failed to find expected category definition in incoming schema %s."),
+						*SchemaDefinition->Id.ToString().ToLower(),
+						*SchemaCategoryDefinition->Id.ToString().ToLower(),
+						*SchemaCategoryDefinition->SchemaCompatibilityServiceAttributeId.ToString().ToLower(),
+						*SchemaCategoryDefinition->Id.ToString().ToLower());
+					return TOnlineResult<FSchemaCategoryInstanceApplyServiceSnapshot>(Errors::InvalidParams());
+				}
+			}
+		}
+		else
+		{
+			UE_LOG(LogOnlineSchema, Warning, TEXT("[ApplyServiceSnapshot] %s.%s.%s: Failed to find definition for SchemaCompatibilityId 0x%08" INT64_X_FMT "."),
+				*SchemaDefinition->Id.ToString().ToLower(),
+				*SchemaCategoryDefinition->Id.ToString().ToLower(),
+				*SchemaCategoryDefinition->SchemaCompatibilityServiceAttributeId.ToString().ToLower(),
+				SchemaCompatibilityIdValue.GetInt64());
+			return TOnlineResult<FSchemaCategoryInstanceApplyServiceSnapshot>(Errors::InvalidParams());
+		}
+	}
+
+	// Unpack service attributes.
+	TMap<FSchemaAttributeId, FSchemaVariant> NewApplicationDataServiceSnapshot;
+	for (TPair<FSchemaServiceAttributeId, FSchemaVariant>& SchemaServiceAttributeData : Params.Attributes)
+	{
+		const FSchemaServiceAttributeId& SchemaServiceAttributeId = SchemaServiceAttributeData.Get<0>();
+		FSchemaVariant& SchemaServiceAttributeValue = SchemaServiceAttributeData.Get<1>();
+
+		if (const FSchemaServiceAttributeDefinition* SchemaServiceAttributeDefinition =
+			SchemaCategoryDefinition->ServiceAttributeDefinitions.Find(SchemaServiceAttributeId))
+		{
+			for (const FSchemaAttributeId& SchemaAttributeId : SchemaServiceAttributeDefinition->SchemaAttributeIds)
+			{
+				// todo: unpack attributes from service attribute.
+				check(SchemaServiceAttributeDefinition->SchemaAttributeIds.Num() == 1);
+
+				const FSchemaAttributeDefinition* SchemaAttributeDefinition = SchemaCategoryDefinition->SchemaAttributeDefinitions.Find(SchemaAttributeId);
+
+				// Attribute definition must be valid due to successfully processing config.
+				check(SchemaAttributeDefinition);
+
+				// Schema changes are handled by surfacing a new schema id in the result object. Don't return SchemaCompatibilityId as an attribute to the application.
+				if (EnumHasAnyFlags(SchemaAttributeDefinition->Flags, ESchemaAttributeFlags::SchemaCompatibilityId))
+				{
+					UE_LOG(LogOnlineSchema, Verbose, TEXT("[ApplyServiceSnapshot] %s.%s.%s: Consuming SchemaCompatibility attribute data: %s."),
+						*SchemaDefinition->Id.ToString().ToLower(),
+						*SchemaCategoryDefinition->Id.ToString().ToLower(),
+						*SchemaAttributeDefinition->Id.ToString().ToLower(),
+						*SchemaServiceAttributeValue.ToLogString());
+				}
+				else
+				{
+					NewApplicationDataServiceSnapshot.Emplace(SchemaAttributeDefinition->Id, MoveTemp(SchemaServiceAttributeValue));
+				}
+			}
+		}
+		else
+		{
+			UE_LOG(LogOnlineSchema, Warning, TEXT("[ApplyServiceSnapshot] %s.%s.%s: Service attribute definition does not exist in schema category."),
+				*SchemaDefinition->Id.ToString().ToLower(),
+				*SchemaCategoryDefinition->Id.ToString().ToLower(),
+				*SchemaServiceAttributeId.ToString().ToLower());
+			return TOnlineResult<FSchemaCategoryInstanceApplyServiceSnapshot>(Errors::InvalidParams());
+		}
+	}
+
+	// Process new snapshot against old snapshot to build application data changes.
+	{
+		// Find attributes that were added or changed.
+		for (const TPair<FSchemaAttributeId, FSchemaVariant>& NewSnapshotAttributeData : NewApplicationDataServiceSnapshot)
+		{
+			const FSchemaServiceAttributeId& NewSnapshotAttributeId = NewSnapshotAttributeData.Get<0>();
+			const FSchemaVariant& NewSnapshotAttributeValue = NewSnapshotAttributeData.Get<1>();
+
+			if (FSchemaVariant* OldSnapshotAttributeValue = ApplicationDataServiceSnapshot.Find(NewSnapshotAttributeId))
+			{
+				Changes.ChangedAttributes.Add(NewSnapshotAttributeId, { MoveTemp(*OldSnapshotAttributeValue), NewSnapshotAttributeValue });
+			}
+			else
+			{
+				Changes.AddedAttributes.Add(NewSnapshotAttributeId, NewSnapshotAttributeValue);
+			}
+		}
+
+		// Add removed attribute ids to changes.
+		Algo::TransformIf(ApplicationDataServiceSnapshot, Changes.RemovedAttributes,
+			[&NewApplicationDataServiceSnapshot](const TPair<FSchemaAttributeId, FSchemaVariant>& OldSnapshotAttributeData) -> bool
+			{
+				const FSchemaServiceAttributeId& OldSnapshotAttributeId = OldSnapshotAttributeData.Get<0>();
+				return !NewApplicationDataServiceSnapshot.Contains(OldSnapshotAttributeId);
+			},
+			[](const TPair<FSchemaAttributeId, FSchemaVariant>& OldSnapshotAttributeData) -> FSchemaAttributeId
+			{
+				const FSchemaServiceAttributeId& OldSnapshotAttributeId = OldSnapshotAttributeData.Get<0>();
+				return OldSnapshotAttributeId;
+			});
+	}
+
+	ApplicationDataServiceSnapshot = MoveTemp(NewApplicationDataServiceSnapshot);
+	return TOnlineResult<FSchemaCategoryInstanceApplyServiceSnapshot>(FSchemaCategoryInstanceApplyServiceSnapshot::Result({MoveTemp(Changes)}));
 }
 
 bool FSchemaCategoryInstance::IsValid() const
 {
-	return SchemaCategoryDefinition != nullptr;
+	return BaseSchemaCategoryDefinition != nullptr;
 }
 
-bool FSchemaCategoryInstance::IsValidAttributeData(const FSchemaAttributeId& Id, const FSchemaVariant& Data)
+bool FSchemaCategoryInstance::IsValidBaseAttributeData(const FSchemaAttributeId& Id, const FSchemaVariant& Data)
 {
 	if (ensure(IsValid()))
 	{
-		if (const FSchemaAttributeDefinition* AttributeDefinition = SchemaCategoryDefinition->SchemaAttributeDefinitions.Find(Id))
+		if (const FSchemaAttributeDefinition* AttributeDefinition = BaseSchemaCategoryDefinition->SchemaAttributeDefinitions.Find(Id))
 		{
 			if (AttributeDefinition->Type == Data.GetType())
 			{
-				return true;
+				if (AttributeDefinition->Type == ESchemaAttributeType::String)
+				{
+					return Data.GetString().Len() <= AttributeDefinition->MaxSize;
+				}
+				else
+				{
+					return true;
+				}
 			}
 			else
 			{
-				UE_LOG(LogOnlineSchema, Verbose, TEXT("[IsValidAttributeData] Schema attribute %s.%s.%s set with invalid type %s. Expected type %s"),
-					*SchemaDefinition->Id.ToString().ToLower(),
-					*SchemaCategoryDefinition->Id.ToString().ToLower(),
+				UE_LOG(LogOnlineSchema, Verbose, TEXT("[IsValidBaseAttributeData] Schema attribute %s.%s.%s set with invalid type %s. Expected type %s"),
+					*BaseSchemaDefinition->Id.ToString().ToLower(),
+					*BaseSchemaCategoryDefinition->Id.ToString().ToLower(),
 					*AttributeDefinition->Id.ToString().ToLower(),
 					*LexToString(Data.GetType()),
 					LexToString(AttributeDefinition->Type));
@@ -468,31 +977,98 @@ bool FSchemaCategoryInstance::IsValidAttributeData(const FSchemaAttributeId& Id,
 		}
 		else
 		{
-			UE_LOG(LogOnlineSchema, Verbose, TEXT("[IsValidAttributeData] Attribute %s not found in schema definition %s.%s"),
+			UE_LOG(LogOnlineSchema, Verbose, TEXT("[IsValidBaseAttributeData] Attribute %s not found in schema definition %s.%s"),
 				*Id.ToString().ToLower(),
-				*SchemaDefinition->Id.ToString().ToLower(),
-				*SchemaCategoryDefinition->Id.ToString().ToLower());
+				*BaseSchemaDefinition->Id.ToString().ToLower(),
+				*BaseSchemaCategoryDefinition->Id.ToString().ToLower());
 			return false;
 		}
 	}
 	else
 	{
-		UE_LOG(LogOnlineSchema, Verbose, TEXT("[IsValidAttributeData] Unable to set attribute %s. Schema is not valid"),
+		UE_LOG(LogOnlineSchema, Verbose, TEXT("[IsValidBaseAttributeData] Unable to set attribute %s. Schema is not valid"),
 			*Id.ToString().ToLower());
 		return false;
 	}
 }
 
+bool FSchemaCategoryInstance::InitializeSchemaDefinition(
+	const FSchemaId& SchemaId,
+	const FSchemaCategoryId& CategoryId,
+	TSharedPtr<const FSchemaDefinition>* OutSchemaDefinition,
+	const FSchemaCategoryDefinition** OutSchemaCategoryDefinition)
+{
+	check(OutSchemaDefinition);
+	check(OutSchemaCategoryDefinition);
+
+	if (TSharedPtr<const FSchemaDefinition> SchemaDefinition = SchemaRegistry->GetDefinition(SchemaId))
+	{
+		if (const FSchemaCategoryDefinition* SchemaCategoryDefinition = SchemaDefinition->Categories.Find(CategoryId))
+		{
+			*OutSchemaDefinition = MoveTemp(SchemaDefinition);
+			*OutSchemaCategoryDefinition = SchemaCategoryDefinition;
+			return true;
+		}
+		else
+		{
+			UE_LOG(LogOnlineSchema, Error, TEXT("[InitializeSchemaDefinition] init error: Unable to find category %s in schema %s"),
+				*CategoryId.ToString().ToLower(),
+				*SchemaId.ToString().ToLower());
+			return false;
+		}
+	}
+	else
+	{
+		UE_LOG(LogOnlineSchema, Error, TEXT("[InitializeSchemaDefinition] init error: Unable to find definition for schema %s"),
+			*SchemaId.ToString().ToLower());
+		return false;
+	}
+}
+
+bool FSchemaCategoryInstance::GetSerializationSchema(
+	const FSchemaDefinition** OutSchemaDefinition,
+	const FSchemaCategoryDefinition** OutSchemaCategoryDefinition) const
+{
+	check(OutSchemaDefinition);
+	check(OutSchemaCategoryDefinition);
+
+	if (IsValid())
+	{
+		if (DerivedSchemaCategoryDefinition)
+		{
+			*OutSchemaDefinition = DerivedSchemaDefinition.Get();
+			*OutSchemaCategoryDefinition = DerivedSchemaCategoryDefinition;
+		}
+		else
+		{
+			*OutSchemaDefinition = BaseSchemaDefinition.Get();
+			*OutSchemaCategoryDefinition = BaseSchemaCategoryDefinition;
+		}
+
+		return true;
+	}
+	else
+	{
+		return false;
+	}
+}
+
 TSharedPtr<const FSchemaDefinition> FSchemaRegistry::GetDefinition(const FSchemaId& SchemaId) const
 {
-	const TSharedRef<const FSchemaDefinition>* Definition = SchemaDefinitions.Find(SchemaId);
+	const TSharedRef<const FSchemaDefinition>* Definition = SchemaDefinitionsById.Find(SchemaId);
+	return Definition ? *Definition : TSharedPtr<const FSchemaDefinition>();
+}
+
+TSharedPtr<const FSchemaDefinition> FSchemaRegistry::GetDefinition(int64 CompatibilityId) const
+{
+	const TSharedRef<const FSchemaDefinition>* Definition = SchemaDefinitionsByCompatibilityId.Find(CompatibilityId);
 	return Definition ? *Definition : TSharedPtr<const FSchemaDefinition>();
 }
 
 bool FSchemaRegistry::IsSchemaChildOf(const FSchemaId& SchemaId, const FSchemaId& ParentSchemaId) const
 {
 	const FSchemaId* FoundSchemaId = nullptr;
-	if (const TSharedRef<const FSchemaDefinition>* Definition = SchemaDefinitions.Find(SchemaId))
+	if (const TSharedRef<const FSchemaDefinition>* Definition = SchemaDefinitionsById.Find(SchemaId))
 	{
 		FoundSchemaId = (* Definition)->ParentSchemaIds.Find(ParentSchemaId);
 	}
diff --git a/Engine/Plugins/Online/Experimental/OnlineServices/Source/OnlineServicesInterface/Private/Online/SchemaTypes.cpp b/Engine/Plugins/Online/Experimental/OnlineServices/Source/OnlineServicesInterface/Private/Online/SchemaTypes.cpp
--- a/Engine/Plugins/Online/Experimental/OnlineServices/Source/OnlineServicesInterface/Private/Online/SchemaTypes.cpp
+++ b/Engine/Plugins/Online/Experimental/OnlineServices/Source/OnlineServicesInterface/Private/Online/SchemaTypes.cpp
@@ -97,6 +97,7 @@ const TCHAR* LexToString(ESchemaAttributeFlags SchemaAttributeFlags)
 	case ESchemaAttributeFlags::None:		return TEXT("None");
 	case ESchemaAttributeFlags::Searchable:	return TEXT("Searchable");
 	case ESchemaAttributeFlags::Public:		return TEXT("Public");
+	case ESchemaAttributeFlags::Private:	return TEXT("Private");
 	}
 }
 
@@ -114,6 +115,10 @@ void LexFromString(ESchemaAttributeFlags& OutSchemaAttributeFlags, const TCHAR*
 	{
 		OutSchemaAttributeFlags = ESchemaAttributeFlags::Public;
 	}
+	else if (FCString::Stricmp(InStr, TEXT("Private")) == 0)
+	{
+		OutSchemaAttributeFlags = ESchemaAttributeFlags::Private;
+	}
 	else if (FCString::Stricmp(InStr, TEXT("SchemaCompatibilityId")) == 0)
 	{
 		OutSchemaAttributeFlags = ESchemaAttributeFlags::SchemaCompatibilityId;
@@ -171,6 +176,7 @@ const TCHAR* LexToString(ESchemaServiceAttributeFlags SchemaAttributeFlags)
 	case ESchemaServiceAttributeFlags::None:		return TEXT("None");
 	case ESchemaServiceAttributeFlags::Searchable:	return TEXT("Searchable");
 	case ESchemaServiceAttributeFlags::Public:		return TEXT("Public");
+	case ESchemaServiceAttributeFlags::Private:		return TEXT("Private");
 	}
 }
 
@@ -188,6 +194,10 @@ void LexFromString(ESchemaServiceAttributeFlags& OutSchemaAttributeFlags, const
 	{
 		OutSchemaAttributeFlags = ESchemaServiceAttributeFlags::Public;
 	}
+	else if (FCString::Stricmp(InStr, TEXT("Private")) == 0)
+	{
+		OutSchemaAttributeFlags = ESchemaServiceAttributeFlags::Private;
+	}
 	else
 	{
 		checkNoEntry();
@@ -319,7 +329,7 @@ bool FSchemaVariant::GetBoolean() const
 	}
 }
 
-FString FSchemaVariant::GetString() const
+const FString& FSchemaVariant::GetString() const
 {
 	if (ensure(VariantType == ESchemaAttributeType::String))
 	{
@@ -327,7 +337,8 @@ FString FSchemaVariant::GetString() const
 	}
 	else
 	{
-		return FString();
+		static FString EmptyString;
+		return EmptyString;
 	}
 }
 
diff --git a/Engine/Plugins/Online/Experimental/OnlineServices/Source/OnlineServicesInterface/Public/Online/Schema.h b/Engine/Plugins/Online/Experimental/OnlineServices/Source/OnlineServicesInterface/Public/Online/Schema.h
--- a/Engine/Plugins/Online/Experimental/OnlineServices/Source/OnlineServicesInterface/Public/Online/Schema.h
+++ b/Engine/Plugins/Online/Experimental/OnlineServices/Source/OnlineServicesInterface/Public/Online/Schema.h
@@ -12,35 +12,54 @@ class FSchemaRegistry;
 class ONLINESERVICESINTERFACE_API FSchemaCategoryInstance
 {
 public:
+	// DerivedSchemaId may remain unset in two situations:
+	// 1. DerivedSchemaId will be populated by attributes from a search result. In this scenario
+	//    the base schema is required to have an attribute flagged as SchemaCompatibilityId so that
+	//    the derived schema can be discovered.
+	// 2. Schema swapping is not enabled. All attributes are defined in the base schema.
 	FSchemaCategoryInstance(
-		const FSchemaId& SchemaId,
+		const FSchemaId& DerivedSchemaId,
 		const FSchemaId& BaseSchemaId,
 		const FSchemaCategoryId& CategoryId,
-		const TSharedRef<const FSchemaRegistry>& Registry);
+		const TSharedRef<const FSchemaRegistry>& SchemaRegistry);
 
-#if 0 // todo
+#if 0
 	TOnlineResult<FSchemaCategoryInstanceChangeSchema> ChangeSchema(FSchemaCategoryInstanceChangeSchema::Params&& Params);
 #endif
-	TOnlineResult<FSchemaCategoryInstanceApplyApplicationChanges> ApplyChanges(FSchemaCategoryInstanceApplyApplicationChanges::Params&& Changes);
-	TOnlineResult<FSchemaCategoryInstanceApplyServiceChanges> ApplyChanges(FSchemaCategoryInstanceApplyServiceChanges::Params&& Changes);
+	TOnlineResult<FSchemaCategoryInstanceApplyApplicationDelta> ApplyApplicationDelta(FSchemaCategoryInstanceApplyApplicationDelta::Params&& Params);
+	TOnlineResult<FSchemaCategoryInstanceApplyServiceSnapshot> ApplyServiceSnapshot(FSchemaCategoryInstanceApplyServiceSnapshot::Params&& Params);
+	// Todo: ApplyServiceDelta
 
-	TSharedPtr<const FSchemaDefinition> GetDefinition() const { return SchemaDefinition; };
+	TSharedPtr<const FSchemaDefinition> GetDerivedDefinition() const { return DerivedSchemaDefinition; };
+	TSharedPtr<const FSchemaDefinition> GetBaseDefinition() const { return BaseSchemaDefinition; };
 
-	// Check whether the category instance initialized successfully.
+	// Check whether the schema is valid.
 	bool IsValid() const;
 
-	// Check whether an application attribute is valid in the schema category.
-	// Intended to verify attributes when there is no backing service for a schema.
-	bool IsValidAttributeData(const FSchemaAttributeId& Id, const FSchemaVariant& Data);
+	// Check whether an application attribute is valid in the base schema.
+	// Intended to be used when a schema is not backed by a service.
+	bool IsValidBaseAttributeData(const FSchemaAttributeId& Id, const FSchemaVariant& Data);
 
 private:
-	TSharedPtr<const FSchemaDefinition> SchemaDefinition;
-	const FSchemaCategoryDefinition* SchemaCategoryDefinition = nullptr;
-
-	// Cache of data in application form. Updates from the application may only be part of a
-	// service field, so the values need to be cached so that the service field can rebuilt when
-	// changes occur.
-	TMap<FSchemaAttributeId, FSchemaVariant> AppDataSnapshot;
+	bool InitializeSchemaDefinition(
+		const FSchemaId& SchemaId,
+		const FSchemaCategoryId& CategoryId,
+		TSharedPtr<const FSchemaDefinition>* OutSchemaDefinition,
+		const FSchemaCategoryDefinition** OutSchemaCategoryDefinition);
+
+	bool GetSerializationSchema(
+		const FSchemaDefinition** OutSchemaDefinition,
+		const FSchemaCategoryDefinition** OutSchemaCategoryDefinition) const;
+
+	TSharedRef<const FSchemaRegistry> SchemaRegistry;
+	TSharedPtr<const FSchemaDefinition> DerivedSchemaDefinition;
+	TSharedPtr<const FSchemaDefinition> BaseSchemaDefinition;
+	const FSchemaCategoryDefinition* DerivedSchemaCategoryDefinition = nullptr;
+	const FSchemaCategoryDefinition* BaseSchemaCategoryDefinition = nullptr;
+	int64 LastSentSchemaCompatibilityId = 0;
+
+	// Snapshot of application data as provided by the service.
+	TMap<FSchemaAttributeId, FSchemaVariant> ApplicationDataServiceSnapshot;
 };
 
 class ONLINESERVICESINTERFACE_API FSchemaRegistry
@@ -50,11 +69,13 @@ public:
 	bool ParseConfig(const FSchemaRegistryDescriptorConfig& Config);
 
 	TSharedPtr<const FSchemaDefinition> GetDefinition(const FSchemaId& SchemaId) const;
+	TSharedPtr<const FSchemaDefinition> GetDefinition(int64 CompatibilityId) const;
 	bool IsSchemaChildOf(const FSchemaId& SchemaId, const FSchemaId& ParentSchemaId) const;
 
 private:
 
-	TMap<FSchemaId, TSharedRef<const FSchemaDefinition>> SchemaDefinitions;
+	TMap<FSchemaId, TSharedRef<const FSchemaDefinition>> SchemaDefinitionsById;
+	TMap<int64, TSharedRef<const FSchemaDefinition>> SchemaDefinitionsByCompatibilityId;
 };
 
 /* UE::Online */ }
diff --git a/Engine/Plugins/Online/Experimental/OnlineServices/Source/OnlineServicesInterface/Public/Online/SchemaTypes.h b/Engine/Plugins/Online/Experimental/OnlineServices/Source/OnlineServicesInterface/Public/Online/SchemaTypes.h
--- a/Engine/Plugins/Online/Experimental/OnlineServices/Source/OnlineServicesInterface/Public/Online/SchemaTypes.h
+++ b/Engine/Plugins/Online/Experimental/OnlineServices/Source/OnlineServicesInterface/Public/Online/SchemaTypes.h
@@ -91,12 +91,16 @@ enum class ESchemaAttributeFlags : uint8
 	  * Attribute is publicly visible.
 	  */
 	Public = 1 << 1,
+	/**
+	  * Attribute is visible to members only.
+	  */
+	Private = 1 << 2,
 	/**
 	  * If multiple schema are supported using the same service interface, a base schema must exist
 	  * with a field marked as schema compatibility id. This field will allow remote clients to
 	  * determine what schema should be used to handle attributes.
 	  */
-	SchemaCompatibilityId = 1 << 2,
+	SchemaCompatibilityId = 1 << 3,
 };
 ENUM_CLASS_FLAGS(ESchemaAttributeFlags);
 ONLINESERVICESINTERFACE_API const TCHAR* LexToString(ESchemaAttributeFlags SchemaAttributeFlags);
@@ -148,6 +152,10 @@ enum class ESchemaServiceAttributeFlags : uint8
 	  * Attribute is publicly visible.
 	  */
 	Public = 1 << 1,
+	/**
+	  * Attribute is visible to members only.
+	  */
+	Private = 1 << 2,
 };
 ENUM_CLASS_FLAGS(ESchemaServiceAttributeFlags);
 ONLINESERVICESINTERFACE_API const TCHAR* LexToString(ESchemaServiceAttributeFlags SchemaServiceAttributeFlags);
@@ -190,10 +198,12 @@ public:
 	FSchemaVariant& operator=(FSchemaVariant&&);
 	FSchemaVariant& operator=(const FSchemaVariant&) = default;
 
-	template<typename ValueType>
-	FSchemaVariant(const ValueType& InData) { Set(InData); }
-	template<typename ValueType>
-	FSchemaVariant(ValueType&& InData) { Set(MoveTemp(InData)); }
+	FSchemaVariant(const TCHAR* InData) { Set(InData); }
+	FSchemaVariant(const FString& InData) { Set(InData); }
+	FSchemaVariant(FString&& InData) { Set(MoveTemp(InData)); }
+	FSchemaVariant(int64 InData) { Set(InData); }
+	FSchemaVariant(double InData) { Set(InData); }
+	FSchemaVariant(bool InData) { Set(InData); }
 	void Set(const TCHAR* AsString);
 	void Set(const FString& AsString);
 	void Set(FString&& AsString);
@@ -206,7 +216,7 @@ public:
 	int64 GetInt64() const;
 	double GetDouble() const;
 	bool GetBoolean() const;
-	FString GetString() const;
+	const FString& GetString() const;
 
 	FString ToLogString() const;
 
@@ -217,7 +227,10 @@ public:
 	ESchemaAttributeType VariantType = ESchemaAttributeType::None;
 };
 
-// Don't allow implicit conversion to FSchemaVariant when calling LexToString.
+/**
+  * Adaptor for handling LexToString with a schema variant.
+  * The adapter prevents unintentional implicit conversion of types to FSchemaVariant when calling LexToString.
+  */
 template <typename T> class FLexToStringAdaptor;
 template <>
 class FLexToStringAdaptor<FSchemaVariant>
@@ -308,6 +321,9 @@ struct FSchemaAttributeDescriptor
 	int32 MaxSize = 0;
 };
 
+/**
+  * Describes a category within a schema.
+  */
 struct FSchemaCategoryDescriptor
 {
 	/**
@@ -323,7 +339,7 @@ struct FSchemaCategoryDescriptor
 	  */
 	FSchemaServiceDescriptorId ServiceDescriptorId;
 	/**
-	  * Attribute definitions.
+	  * Attribute descriptors.
 	  */
 	TArray<FSchemaAttributeDescriptor> Attributes;
 };
@@ -378,11 +394,10 @@ struct FSchemaServiceAttributeDefinition
 	  */
 	ESchemaServiceAttributeSupportedTypeFlags Type;
 	/**
-	  * Declares the maximum allowed size for the service attribute.
-	  *
-	  * Applicable when using a variable sized type such as String.
+	  * Additional behavior required by the attribute. Examples include making an attribute
+	  * public, private or available to search.
 	  */
-	ESchemaAttributeFlags Flags = ESchemaAttributeFlags::None;
+	ESchemaServiceAttributeFlags Flags = ESchemaServiceAttributeFlags::None;
 	/**
 	  * Declares the maximum allowed size for the service attribute.
 	  *
@@ -410,7 +425,7 @@ struct FSchemaAttributeDefinition
 	ESchemaAttributeType Type = ESchemaAttributeType::None;
 	/**
 	  * Additional behavior required by the attribute. Examples include making an attribute
-	  * public or available to search.
+	  * public, private or available to search.
 	  */
 	ESchemaAttributeFlags Flags;
 	/**
@@ -425,6 +440,10 @@ struct FSchemaAttributeDefinition
 	FSchemaServiceAttributeId ServiceAttributeId;
 };
 
+/**
+  * The runtime definition of a schema category.
+  * The contained data can be used for validation and translation.
+  */
 struct FSchemaCategoryDefinition
 {
 	/**
@@ -433,6 +452,18 @@ struct FSchemaCategoryDefinition
 	  * Example: "Lobby" vs "LobbyMember"
 	  */
 	FSchemaCategoryId Id;
+	/**
+	  * The Id of the attribute containing the schema compatibility id.
+	  * Changes to the schema compatibility attribute are not allowed externally. This id is used
+	  * to fail attempted external changes.
+	  */
+	FSchemaAttributeId SchemaCompatibilityAttributeId;
+	/**
+	  * The Id of the service attribute containing the schema compatibility id.
+	  * This Id is used to detect the schema to use when handling search results and to handle
+	  * schema changes when the remote owner makes a change.
+	  */
+	FSchemaServiceAttributeId SchemaCompatibilityServiceAttributeId;
 	/**
 	  * Mapping of schema attribute ids to attribute definitions.
 	  */
@@ -469,8 +500,7 @@ struct FSchemaDefinition
 };
 
 /**
-  * The runtime definition of a schema.
-  * The contained data can be used for validation and translation.
+  * The set of data used to write a variant to a service.
   */
 struct FSchemaServiceAttributeData
 {
@@ -484,25 +514,9 @@ struct FSchemaServiceAttributeData
 	  */
 	ESchemaAttributeFlags Flags = ESchemaAttributeFlags::None;
 	/**
-	  * The attribute data.
+	  * The attribute value.
 	  */
-	FSchemaVariant Data;
-};
-
-/**
-  * The runtime definition of a schema.
-  * The contained data can be used for validation and translation.
-  */
-struct FSchemaAttributeData
-{
-	/**
-	  * The id to associate with this attribute.
-	  */
-	FSchemaAttributeId Id;
-	/**
-	  * The attribute data.
-	  */
-	FSchemaVariant Data;
+	FSchemaVariant Value;
 };
 
 /**
@@ -511,17 +525,13 @@ struct FSchemaAttributeData
 struct FSchemaServiceChanges
 {
 	/**
-	  * Added attributes and their values.
+	  * Added or changed attributes and their values.
 	  */
-	TMap<FSchemaServiceAttributeId, FSchemaServiceAttributeData> AddedAttributes;
+	TMap<FSchemaServiceAttributeId, FSchemaServiceAttributeData> UpdatedAttributes;
 	/**
-	  * Removed attributes with their previous values.
+	  * Removed attribute ids.
 	  */
 	TSet<FSchemaAttributeId> RemovedAttributes;
-	/**
-	  * Changed attributes with their old and new values.
-	  */
-	TMap<FSchemaServiceAttributeId, TPair<FSchemaServiceAttributeData, FSchemaServiceAttributeData>> ChangedAttributes;
 };
 
 /**
@@ -530,17 +540,21 @@ struct FSchemaServiceChanges
 struct FSchemaApplicationChanges
 {
 	/**
-	  * Added attributes and their values.
+	  * Change to schema id.
 	  */
-	TMap<FSchemaAttributeId, FSchemaAttributeData> AddedAttributes;
+	TOptional<FSchemaId> SchemaId;
 	/**
-	  * Removed attributes.
+	  * Added attributes.
 	  */
-	TSet<FSchemaAttributeId> RemovedAttributes;
+	TMap<FSchemaAttributeId, FSchemaVariant> AddedAttributes;
 	/**
 	  * Changed attributes with their old and new values.
 	  */
-	TMap<FSchemaAttributeId, TPair<FSchemaAttributeData, FSchemaAttributeData>> ChangedAttributes;
+	TMap<FSchemaAttributeId, TPair<FSchemaVariant, FSchemaVariant>> ChangedAttributes;
+	/**
+	  * Removed attribute ids.
+	  */
+	TSet<FSchemaAttributeId> RemovedAttributes;
 };
 
 #if 0 // todo
@@ -568,18 +582,16 @@ struct FSchemaCategoryInstanceChangeSchema
 };
 #endif
 
-struct FSchemaCategoryInstanceApplyApplicationChanges
+struct FSchemaCategoryInstanceApplyApplicationDelta
 {
-	static constexpr TCHAR Name[] = TEXT("ApplyApplicationChanges");
+	static constexpr TCHAR Name[] = TEXT("ApplyApplicationDelta");
 
 	struct Params
 	{
 		/** Attributes notified from the application. */
-		TArray<FSchemaAttributeData> MutatedAttributes;
+		TMap<FSchemaAttributeId, FSchemaVariant> UpdatedAttributes;
 		/** Removed attributes notified from the application. */
-		TArray<FSchemaAttributeId> RemovedAttributes;
-		/** Whether the attributes represent a full snapshot or a delta from a prior update. */
-		bool bIsDelta = false;
+		TSet<FSchemaAttributeId> RemovedAttributes;
 	};
 
 	struct Result
@@ -592,18 +604,14 @@ struct FSchemaCategoryInstanceApplyApplicationChanges
 	};
 };
 
-struct FSchemaCategoryInstanceApplyServiceChanges
+struct FSchemaCategoryInstanceApplyServiceSnapshot
 {
-	static constexpr TCHAR Name[] = TEXT("ApplyServiceChanges");
+	static constexpr TCHAR Name[] = TEXT("ApplyServiceSnapshot");
 
 	struct Params
 	{
-		/** Attributes notified from the service. */
-		TArray<FSchemaServiceAttributeData> MutatedAttributes;
-		/** Removed attributes notified from the service. */
-		TArray<FSchemaAttributeId> RemovedAttributes;
-		/** Whether the attributes represent a full snapshot or a delta from a prior update. */
-		bool bIsDelta = false;
+		/** Full attribute snapshot from the service. */
+		TMap<FSchemaServiceAttributeId, FSchemaVariant> Attributes;
 	};
 
 	struct Result
