OSSv2: Rename FOnlineAccountIdHandle to FAccountId.

diff --git a/Engine/Plugins/Online/Experimental/OnlineServices/Source/OnlineServicesCommon/Private/Online/AchievementsCommon.cpp b/Engine/Plugins/Online/Experimental/OnlineServices/Source/OnlineServicesCommon/Private/Online/AchievementsCommon.cpp
--- a/Engine/Plugins/Online/Experimental/OnlineServices/Source/OnlineServicesCommon/Private/Online/AchievementsCommon.cpp
+++ b/Engine/Plugins/Online/Experimental/OnlineServices/Source/OnlineServicesCommon/Private/Online/AchievementsCommon.cpp
@@ -141,7 +141,7 @@ void FAchievementsCommon::UnlockAchievementsByStats(const FStatsUpdated& StatsUp
 	}
 
 	TArray<FString> StatNames;
-	TArray<FOnlineAccountIdHandle> UserIds;
+	TArray<FAccountId> UserIds;
 	for (const FUserStats& UserStats : StatsUpdated.UpdateUsersStats)
 	{
 		for (const TPair<FString, FStatValue>& StatPair : UserStats.Stats)
@@ -197,7 +197,7 @@ void FAchievementsCommon::UnlockAchievementsByStats(const FStatsUpdated& StatsUp
 	});
 }
 
-void FAchievementsCommon::ExecuteUnlockRulesRelatedToStat(const FOnlineAccountIdHandle& UserId, const FString& StatName, const TMap<FString, FStatValue>& Stats, TArray<FString>& OutAchievementsToUnlock)
+void FAchievementsCommon::ExecuteUnlockRulesRelatedToStat(const FAccountId& UserId, const FString& StatName, const TMap<FString, FStatValue>& Stats, TArray<FString>& OutAchievementsToUnlock)
 {
 	for (const FAchievementUnlockRule& AchievementUnlockRule : AchievementUnlockRules)
 	{
@@ -251,7 +251,7 @@ bool FAchievementsCommon::MeetUnlockCondition(FAchievementUnlockRule Achievement
 	return true;
 }
 
-bool FAchievementsCommon::IsUnlocked(const FOnlineAccountIdHandle& UserId, const FString& AchievementName) const
+bool FAchievementsCommon::IsUnlocked(const FAccountId& UserId, const FString& AchievementName) const
 {
 	FGetAchievementState::Params Params;
 	Params.LocalUserId = UserId;
diff --git a/Engine/Plugins/Online/Experimental/OnlineServices/Source/OnlineServicesCommon/Private/Online/AuthCommon.cpp b/Engine/Plugins/Online/Experimental/OnlineServices/Source/OnlineServicesCommon/Private/Online/AuthCommon.cpp
--- a/Engine/Plugins/Online/Experimental/OnlineServices/Source/OnlineServicesCommon/Private/Online/AuthCommon.cpp
+++ b/Engine/Plugins/Online/Experimental/OnlineServices/Source/OnlineServicesCommon/Private/Online/AuthCommon.cpp
@@ -17,7 +17,7 @@ TSharedPtr<FAccountInfo> FAccountInfoRegistry::Find(FPlatformUserId PlatformUser
 	return FoundPtr ? TSharedPtr<FAccountInfo>(*FoundPtr) : TSharedPtr<FAccountInfo>();
 }
 
-TSharedPtr<FAccountInfo> FAccountInfoRegistry::Find(FOnlineAccountIdHandle AccountIdHandle) const
+TSharedPtr<FAccountInfo> FAccountInfoRegistry::Find(FAccountId AccountIdHandle) const
 {
 	FReadScopeLock Lock(IndexLock);
 	const TSharedRef<FAccountInfo>* FoundPtr = AuthDataByOnlineAccountIdHandle.Find(AccountIdHandle);
@@ -189,7 +189,7 @@ TOnlineEvent<void(const FAuthAccountAttributesChanged&)> FAuthCommon::OnAccountA
 	return OnAuthAccountAttributesChangedEvent;
 }
 
-bool FAuthCommon::IsLoggedIn(const FOnlineAccountIdHandle& AccountId) const
+bool FAuthCommon::IsLoggedIn(const FAccountId& AccountId) const
 {
 	const TSharedPtr<FAccountInfo> AccountInfo = GetAccountInfoRegistry().Find(AccountId);
 	return (AccountInfo && IsOnlineStatus(AccountInfo->LoginStatus));
diff --git a/Engine/Plugins/Online/Experimental/OnlineServices/Source/OnlineServicesCommon/Private/Online/LobbiesCommonTests.cpp b/Engine/Plugins/Online/Experimental/OnlineServices/Source/OnlineServicesCommon/Private/Online/LobbiesCommonTests.cpp
--- a/Engine/Plugins/Online/Experimental/OnlineServices/Source/OnlineServicesCommon/Private/Online/LobbiesCommonTests.cpp
+++ b/Engine/Plugins/Online/Experimental/OnlineServices/Source/OnlineServicesCommon/Private/Online/LobbiesCommonTests.cpp
@@ -117,8 +117,8 @@ bool FClientLobbyDataTest::RunTest(const FString& Parameters)
 		FOnlineLobbyIdHandle LobbyId(EOnlineServices::Null, 1);
 		FClientLobbyData ClientData(LobbyId);
 
-		FOnlineAccountIdHandle User1(EOnlineServices::Null, 1);
-		FOnlineAccountIdHandle User2(EOnlineServices::Null, 2);
+		FAccountId User1(EOnlineServices::Null, 1);
+		FAccountId User2(EOnlineServices::Null, 2);
 		FName LobbySchemaName = TEXT("SchemaName");
 		int32 LobbyMaxMembers = 5;
 		ELobbyJoinPolicy LobbyJoinPolicy = ELobbyJoinPolicy::PublicAdvertised;
@@ -137,7 +137,7 @@ bool FClientLobbyDataTest::RunTest(const FString& Parameters)
 
 		// 1.1 setup.
 		{
-			TMap<FOnlineAccountIdHandle, TSharedRef<FClientLobbyMemberSnapshot>> LobbyMemberSnapshots;
+			TMap<FAccountId, TSharedRef<FClientLobbyMemberSnapshot>> LobbyMemberSnapshots;
 			{
 				TSharedRef<FClientLobbyMemberSnapshot> MemberSnapshot = MakeShared<FClientLobbyMemberSnapshot>();
 				MemberSnapshot->AccountId = User1;
@@ -158,7 +158,7 @@ bool FClientLobbyDataTest::RunTest(const FString& Parameters)
 			LobbySnapshot.Attributes.Add(LobbyAttribute1Key, LobbyAttribute1Value);
 			LobbySnapshot.Members = {User1, User2};
 
-			TMap<FOnlineAccountIdHandle, ELobbyMemberLeaveReason> LeaveReasons;
+			TMap<FAccountId, ELobbyMemberLeaveReason> LeaveReasons;
 			FApplyLobbyUpdateResult Result = ClientData.ApplyLobbyUpdateFromServiceSnapshot(MoveTemp(LobbySnapshot), MoveTemp(LobbyMemberSnapshots), MoveTemp(LeaveReasons), &LobbyEvents);
 			UTEST_EQUAL("Check no local members left.", Result.LeavingLocalMembers.Num(), 0);
 		}
@@ -187,7 +187,7 @@ bool FClientLobbyDataTest::RunTest(const FString& Parameters)
 			LobbyMemberAttribute1Value.Set(TEXT("ModifiedMemberAttribute1"));
 			LobbySchemaName = TEXT("ModifiedSchemaName");
 
-			TMap<FOnlineAccountIdHandle, TSharedRef<FClientLobbyMemberSnapshot>> LobbyMemberSnapshots;
+			TMap<FAccountId, TSharedRef<FClientLobbyMemberSnapshot>> LobbyMemberSnapshots;
 			{
 				TSharedRef<FClientLobbyMemberSnapshot> MemberSnapshot = MakeShared<FClientLobbyMemberSnapshot>();
 				MemberSnapshot->AccountId = User1;
@@ -208,7 +208,7 @@ bool FClientLobbyDataTest::RunTest(const FString& Parameters)
 			LobbySnapshot.Attributes.Add(LobbyAttribute1Key, LobbyAttribute1Value);
 			LobbySnapshot.Members = {User1, User2};
 
-			TMap<FOnlineAccountIdHandle, ELobbyMemberLeaveReason> LeaveReasons;
+			TMap<FAccountId, ELobbyMemberLeaveReason> LeaveReasons;
 			FApplyLobbyUpdateResult Result = ClientData.ApplyLobbyUpdateFromServiceSnapshot(MoveTemp(LobbySnapshot), MoveTemp(LobbyMemberSnapshots), MoveTemp(LeaveReasons), &LobbyEvents);
 			UTEST_EQUAL("Check no local members left.", Result.LeavingLocalMembers.Num(), 0);
 		}
@@ -234,7 +234,7 @@ bool FClientLobbyDataTest::RunTest(const FString& Parameters)
 
 		// 1.3 setup.
 		{
-			TMap<FOnlineAccountIdHandle, TSharedRef<FClientLobbyMemberSnapshot>> LobbyMemberSnapshots;
+			TMap<FAccountId, TSharedRef<FClientLobbyMemberSnapshot>> LobbyMemberSnapshots;
 			{
 				TSharedRef<FClientLobbyMemberSnapshot> MemberSnapshot = MakeShared<FClientLobbyMemberSnapshot>();
 				MemberSnapshot->AccountId = User1;
@@ -248,7 +248,7 @@ bool FClientLobbyDataTest::RunTest(const FString& Parameters)
 			LobbySnapshot.JoinPolicy = LobbyJoinPolicy;
 			LobbySnapshot.Members = {User1};
 
-			TMap<FOnlineAccountIdHandle, ELobbyMemberLeaveReason> LeaveReasons;
+			TMap<FAccountId, ELobbyMemberLeaveReason> LeaveReasons;
 			FApplyLobbyUpdateResult Result = ClientData.ApplyLobbyUpdateFromServiceSnapshot(MoveTemp(LobbySnapshot), MoveTemp(LobbyMemberSnapshots), MoveTemp(LeaveReasons), &LobbyEvents);
 			UTEST_EQUAL("Check no local members left.", Result.LeavingLocalMembers.Num(), 0);
 		}
@@ -270,8 +270,8 @@ bool FClientLobbyDataTest::RunTest(const FString& Parameters)
 		FOnlineLobbyIdHandle LobbyId(EOnlineServices::Null, 1);
 		FClientLobbyData ClientData(LobbyId);
 
-		FOnlineAccountIdHandle User1(EOnlineServices::Null, 1);
-		FOnlineAccountIdHandle User2(EOnlineServices::Null, 2);
+		FAccountId User1(EOnlineServices::Null, 1);
+		FAccountId User2(EOnlineServices::Null, 2);
 		FName LobbySchemaName = TEXT("SchemaName");
 		int32 LobbyMaxMembers = 5;
 		ELobbyJoinPolicy LobbyJoinPolicy = ELobbyJoinPolicy::PublicAdvertised;
@@ -290,7 +290,7 @@ bool FClientLobbyDataTest::RunTest(const FString& Parameters)
 
 		// 2.1 setup
 		{
-			TMap<FOnlineAccountIdHandle, TSharedRef<FClientLobbyMemberSnapshot>> LobbyMemberSnapshots;
+			TMap<FAccountId, TSharedRef<FClientLobbyMemberSnapshot>> LobbyMemberSnapshots;
 			{
 				TSharedRef<FClientLobbyMemberSnapshot> MemberSnapshot = MakeShared<FClientLobbyMemberSnapshot>();
 				MemberSnapshot->AccountId = User1;
@@ -311,7 +311,7 @@ bool FClientLobbyDataTest::RunTest(const FString& Parameters)
 			LobbySnapshot.Attributes.Add(LobbyAttribute1Key, LobbyAttribute1Value);
 			LobbySnapshot.Members = {User1, User2};
 
-			TMap<FOnlineAccountIdHandle, ELobbyMemberLeaveReason> LeaveReasons;
+			TMap<FAccountId, ELobbyMemberLeaveReason> LeaveReasons;
 			FApplyLobbyUpdateResult Result = ClientData.ApplyLobbyUpdateFromServiceSnapshot(MoveTemp(LobbySnapshot), MoveTemp(LobbyMemberSnapshots), MoveTemp(LeaveReasons), &LobbyEvents);
 			UTEST_EQUAL("Check no members left.", Result.LeavingLocalMembers.Num(), 0);
 		}
@@ -349,7 +349,7 @@ bool FClientLobbyDataTest::RunTest(const FString& Parameters)
 			LobbyMemberAttribute1Value.Set(TEXT("ModifiedMemberAttribute1"));
 			LobbySchemaName = TEXT("ModifiedSchemaName");
 
-			TMap<FOnlineAccountIdHandle, TSharedRef<FClientLobbyMemberSnapshot>> LobbyMemberSnapshots;
+			TMap<FAccountId, TSharedRef<FClientLobbyMemberSnapshot>> LobbyMemberSnapshots;
 			{
 				TSharedRef<FClientLobbyMemberSnapshot> MemberSnapshot = MakeShared<FClientLobbyMemberSnapshot>();
 				MemberSnapshot->AccountId = User1;
@@ -370,7 +370,7 @@ bool FClientLobbyDataTest::RunTest(const FString& Parameters)
 			LobbySnapshot.Attributes.Add(LobbyAttribute1Key, LobbyAttribute1Value);
 			LobbySnapshot.Members = {User1, User2};
 
-			TMap<FOnlineAccountIdHandle, ELobbyMemberLeaveReason> LeaveReasons;
+			TMap<FAccountId, ELobbyMemberLeaveReason> LeaveReasons;
 			FApplyLobbyUpdateResult Result = ClientData.ApplyLobbyUpdateFromServiceSnapshot(MoveTemp(LobbySnapshot), MoveTemp(LobbyMemberSnapshots), MoveTemp(LeaveReasons), &LobbyEvents);
 			UTEST_EQUAL("Check no local members left.", Result.LeavingLocalMembers.Num(), 0);
 		}
@@ -399,7 +399,7 @@ bool FClientLobbyDataTest::RunTest(const FString& Parameters)
 		{
 			EventCapture.Empty();
 
-			TMap<FOnlineAccountIdHandle, TSharedRef<FClientLobbyMemberSnapshot>> LobbyMemberSnapshots;
+			TMap<FAccountId, TSharedRef<FClientLobbyMemberSnapshot>> LobbyMemberSnapshots;
 			{
 				TSharedRef<FClientLobbyMemberSnapshot> MemberSnapshot = MakeShared<FClientLobbyMemberSnapshot>();
 				MemberSnapshot->AccountId = User2;
@@ -413,7 +413,7 @@ bool FClientLobbyDataTest::RunTest(const FString& Parameters)
 			LobbySnapshot.JoinPolicy = LobbyJoinPolicy;
 			LobbySnapshot.Members = {User2};
 
-			TMap<FOnlineAccountIdHandle, ELobbyMemberLeaveReason> LeaveReasons;
+			TMap<FAccountId, ELobbyMemberLeaveReason> LeaveReasons;
 			LeaveReasons.Add(User1, ELobbyMemberLeaveReason::Kicked);
 
 			FApplyLobbyUpdateResult Result = ClientData.ApplyLobbyUpdateFromServiceSnapshot(MoveTemp(LobbySnapshot), MoveTemp(LobbyMemberSnapshots), MoveTemp(LeaveReasons), &LobbyEvents);
@@ -457,15 +457,15 @@ bool FClientLobbyDataTest::RunTest(const FString& Parameters)
 		FOnlineLobbyIdHandle LobbyId(EOnlineServices::Null, 1);
 		FClientLobbyData ClientData(LobbyId);
 
-		FOnlineAccountIdHandle User1(EOnlineServices::Null, 1);
-		FOnlineAccountIdHandle User2(EOnlineServices::Null, 2);
+		FAccountId User1(EOnlineServices::Null, 1);
+		FAccountId User2(EOnlineServices::Null, 2);
 		FName LobbySchemaName = TEXT("SchemaName");
 		int32 LobbyMaxMembers = 5;
 		ELobbyJoinPolicy LobbyJoinPolicy = ELobbyJoinPolicy::PublicAdvertised;
 
 		// 3.1 setup - initialize lobby data with initial snapshot.
 		{
-			TMap<FOnlineAccountIdHandle, TSharedRef<FClientLobbyMemberSnapshot>> LobbyMemberSnapshots;
+			TMap<FAccountId, TSharedRef<FClientLobbyMemberSnapshot>> LobbyMemberSnapshots;
 			{
 				TSharedRef<FClientLobbyMemberSnapshot> MemberSnapshot = MakeShared<FClientLobbyMemberSnapshot>();
 				MemberSnapshot->AccountId = User1;
@@ -484,7 +484,7 @@ bool FClientLobbyDataTest::RunTest(const FString& Parameters)
 			LobbySnapshot.JoinPolicy = LobbyJoinPolicy;
 			LobbySnapshot.Members = {User1, User2};
 
-			TMap<FOnlineAccountIdHandle, ELobbyMemberLeaveReason> LeaveReasons;
+			TMap<FAccountId, ELobbyMemberLeaveReason> LeaveReasons;
 			FApplyLobbyUpdateResult Result = ClientData.ApplyLobbyUpdateFromServiceSnapshot(MoveTemp(LobbySnapshot), MoveTemp(LobbyMemberSnapshots), MoveTemp(LeaveReasons), &LobbyEvents);
 			UTEST_EQUAL("Check no local members left.", Result.LeavingLocalMembers.Num(), 0);
 		}
@@ -501,7 +501,7 @@ bool FClientLobbyDataTest::RunTest(const FString& Parameters)
 		{
 			EventCapture.Empty();
 
-			TMap<FOnlineAccountIdHandle, TSharedRef<FClientLobbyMemberSnapshot>> LobbyMemberSnapshots;
+			TMap<FAccountId, TSharedRef<FClientLobbyMemberSnapshot>> LobbyMemberSnapshots;
 			{
 				TSharedRef<FClientLobbyMemberSnapshot> MemberSnapshot = MakeShared<FClientLobbyMemberSnapshot>();
 				MemberSnapshot->AccountId = User1;
@@ -515,7 +515,7 @@ bool FClientLobbyDataTest::RunTest(const FString& Parameters)
 			LobbySnapshot.JoinPolicy = LobbyJoinPolicy;
 			LobbySnapshot.Members = {User1};
 
-			TMap<FOnlineAccountIdHandle, ELobbyMemberLeaveReason> LeaveReasons;
+			TMap<FAccountId, ELobbyMemberLeaveReason> LeaveReasons;
 			FApplyLobbyUpdateResult Result = ClientData.ApplyLobbyUpdateFromServiceSnapshot(MoveTemp(LobbySnapshot), MoveTemp(LobbyMemberSnapshots), MoveTemp(LeaveReasons), &LobbyEvents);
 			UTEST_EQUAL("Check no local members left.", Result.LeavingLocalMembers.Num(), 0);
 		}
@@ -531,7 +531,7 @@ bool FClientLobbyDataTest::RunTest(const FString& Parameters)
 		{
 			EventCapture.Empty();
 
-			TMap<FOnlineAccountIdHandle, TSharedRef<FClientLobbyMemberSnapshot>> LobbyMemberSnapshots;
+			TMap<FAccountId, TSharedRef<FClientLobbyMemberSnapshot>> LobbyMemberSnapshots;
 			{
 				TSharedRef<FClientLobbyMemberSnapshot> MemberSnapshot = MakeShared<FClientLobbyMemberSnapshot>();
 				MemberSnapshot->AccountId = User1;
@@ -550,7 +550,7 @@ bool FClientLobbyDataTest::RunTest(const FString& Parameters)
 			LobbySnapshot.JoinPolicy = LobbyJoinPolicy;
 			LobbySnapshot.Members = {User1};
 
-			TMap<FOnlineAccountIdHandle, ELobbyMemberLeaveReason> LeaveReasons;
+			TMap<FAccountId, ELobbyMemberLeaveReason> LeaveReasons;
 			FApplyLobbyUpdateResult Result = ClientData.ApplyLobbyUpdateFromServiceSnapshot(MoveTemp(LobbySnapshot), MoveTemp(LobbyMemberSnapshots), MoveTemp(LeaveReasons), &LobbyEvents);
 			UTEST_EQUAL("Check no local members left.", Result.LeavingLocalMembers.Num(), 0);
 		}
@@ -562,7 +562,7 @@ bool FClientLobbyDataTest::RunTest(const FString& Parameters)
 		{
 			EventCapture.Empty();
 
-			TMap<FOnlineAccountIdHandle, TSharedRef<FClientLobbyMemberSnapshot>> LobbyMemberSnapshots;
+			TMap<FAccountId, TSharedRef<FClientLobbyMemberSnapshot>> LobbyMemberSnapshots;
 			{
 				TSharedRef<FClientLobbyMemberSnapshot> MemberSnapshot = MakeShared<FClientLobbyMemberSnapshot>();
 				MemberSnapshot->AccountId = User1;
@@ -581,7 +581,7 @@ bool FClientLobbyDataTest::RunTest(const FString& Parameters)
 			LobbySnapshot.JoinPolicy = LobbyJoinPolicy;
 			LobbySnapshot.Members = {User1, User2};
 
-			TMap<FOnlineAccountIdHandle, ELobbyMemberLeaveReason> LeaveReasons;
+			TMap<FAccountId, ELobbyMemberLeaveReason> LeaveReasons;
 			FApplyLobbyUpdateResult Result = ClientData.ApplyLobbyUpdateFromServiceSnapshot(MoveTemp(LobbySnapshot), MoveTemp(LobbyMemberSnapshots), MoveTemp(LeaveReasons), &LobbyEvents);
 			UTEST_EQUAL("Check no local members left.", Result.LeavingLocalMembers.Num(), 0);
 		}
@@ -600,8 +600,8 @@ bool FClientLobbyDataTest::RunTest(const FString& Parameters)
 		FOnlineLobbyIdHandle LobbyId(EOnlineServices::Null, 1);
 		FClientLobbyData ClientData(LobbyId);
 
-		FOnlineAccountIdHandle User1(EOnlineServices::Null, 1);
-		FOnlineAccountIdHandle User2(EOnlineServices::Null, 2);
+		FAccountId User1(EOnlineServices::Null, 1);
+		FAccountId User2(EOnlineServices::Null, 2);
 		FName LobbySchemaName = TEXT("SchemaName");
 		int32 LobbyMaxMembers = 5;
 		ELobbyJoinPolicy LobbyJoinPolicy = ELobbyJoinPolicy::PublicAdvertised;
@@ -620,7 +620,7 @@ bool FClientLobbyDataTest::RunTest(const FString& Parameters)
 
 		// 4.1 setup. Local member in snapshot with remote member to simulate lobby join.
 		{
-			TMap<FOnlineAccountIdHandle, TSharedRef<FClientLobbyMemberSnapshot>> LobbyMemberSnapshots;
+			TMap<FAccountId, TSharedRef<FClientLobbyMemberSnapshot>> LobbyMemberSnapshots;
 			{
 				TSharedRef<FClientLobbyMemberSnapshot> MemberSnapshot = MakeShared<FClientLobbyMemberSnapshot>();
 				MemberSnapshot->AccountId = User1;
@@ -641,7 +641,7 @@ bool FClientLobbyDataTest::RunTest(const FString& Parameters)
 			LobbySnapshot.Attributes.Add(LobbyAttribute1Key, LobbyAttribute1Value);
 			LobbySnapshot.Members = {User1, User2};
 
-			TMap<FOnlineAccountIdHandle, ELobbyMemberLeaveReason> LeaveReasons;
+			TMap<FAccountId, ELobbyMemberLeaveReason> LeaveReasons;
 			FApplyLobbyUpdateResult Result = ClientData.ApplyLobbyUpdateFromServiceSnapshot(MoveTemp(LobbySnapshot), MoveTemp(LobbyMemberSnapshots), MoveTemp(LeaveReasons), &LobbyEvents);
 			UTEST_EQUAL("Check no local members left.", Result.LeavingLocalMembers.Num(), 0);
 		}
@@ -779,8 +779,8 @@ bool FClientLobbyDataTest::RunTest(const FString& Parameters)
 		FOnlineLobbyIdHandle LobbyId(EOnlineServices::Null, 1);
 		FClientLobbyData ClientData(LobbyId);
 
-		FOnlineAccountIdHandle User1(EOnlineServices::Null, 1);
-		FOnlineAccountIdHandle User2(EOnlineServices::Null, 2);
+		FAccountId User1(EOnlineServices::Null, 1);
+		FAccountId User2(EOnlineServices::Null, 2);
 		FName LobbySchemaName = TEXT("SchemaName");
 		int32 LobbyMaxMembers = 5;
 		ELobbyJoinPolicy LobbyJoinPolicy = ELobbyJoinPolicy::PublicAdvertised;
@@ -799,7 +799,7 @@ bool FClientLobbyDataTest::RunTest(const FString& Parameters)
 
 		// 5.1 setup.
 		{
-			TMap<FOnlineAccountIdHandle, TSharedRef<FClientLobbyMemberSnapshot>> LobbyMemberSnapshots;
+			TMap<FAccountId, TSharedRef<FClientLobbyMemberSnapshot>> LobbyMemberSnapshots;
 			{
 				TSharedRef<FClientLobbyMemberSnapshot> MemberSnapshot = MakeShared<FClientLobbyMemberSnapshot>();
 				MemberSnapshot->AccountId = User1;
@@ -820,7 +820,7 @@ bool FClientLobbyDataTest::RunTest(const FString& Parameters)
 			LobbySnapshot.Attributes.Add(LobbyAttribute1Key, LobbyAttribute1Value);
 			LobbySnapshot.Members = {User1, User2};
 
-			TMap<FOnlineAccountIdHandle, ELobbyMemberLeaveReason> LeaveReasons;
+			TMap<FAccountId, ELobbyMemberLeaveReason> LeaveReasons;
 			FApplyLobbyUpdateResult Result = ClientData.ApplyLobbyUpdateFromServiceSnapshot(MoveTemp(LobbySnapshot), MoveTemp(LobbyMemberSnapshots), MoveTemp(LeaveReasons), &LobbyEvents);
 			UTEST_EQUAL("Check no local members left.", Result.LeavingLocalMembers.Num(), 0);
 		}
@@ -1123,7 +1123,7 @@ TFuture<int> AwaitNextGameTick()
 
 TFuture<TDefaultErrorResultInternal<FOnlineLobbyIdHandle>> AwaitInvitation(
 	FLobbyEvents& LobbyEvents,
-	FOnlineAccountIdHandle TargetAccountId,
+	FAccountId TargetAccountId,
 	FOnlineLobbyIdHandle LobbyId,
 	float TimeoutSeconds)
 {
diff --git a/Engine/Plugins/Online/Experimental/OnlineServices/Source/OnlineServicesCommon/Private/Online/LobbiesCommonTypes.cpp b/Engine/Plugins/Online/Experimental/OnlineServices/Source/OnlineServicesCommon/Private/Online/LobbiesCommonTypes.cpp
--- a/Engine/Plugins/Online/Experimental/OnlineServices/Source/OnlineServicesCommon/Private/Online/LobbiesCommonTypes.cpp
+++ b/Engine/Plugins/Online/Experimental/OnlineServices/Source/OnlineServicesCommon/Private/Online/LobbiesCommonTypes.cpp
@@ -103,7 +103,7 @@ FClientLobbyData::FClientLobbyData(FOnlineLobbyIdHandle LobbyId)
 	PublicData->LobbyId = LobbyId;
 }
 
-TSharedPtr<const FClientLobbyMemberSnapshot> FClientLobbyData::GetMemberData(FOnlineAccountIdHandle MemberAccountId) const
+TSharedPtr<const FClientLobbyMemberSnapshot> FClientLobbyData::GetMemberData(FAccountId MemberAccountId) const
 {
 	const TSharedRef<FClientLobbyMemberSnapshot>* FoundMemberData = MemberDataStorage.Find(MemberAccountId);
 	return FoundMemberData ? *FoundMemberData : TSharedPtr<const FClientLobbyMemberSnapshot>();
@@ -111,8 +111,8 @@ TSharedPtr<const FClientLobbyMemberSnapshot> FClientLobbyData::GetMemberData(FOn
 
 FApplyLobbyUpdateResult FClientLobbyData::ApplyLobbyUpdateFromServiceSnapshot(
 	FClientLobbySnapshot&& LobbySnapshot,
-	TMap<FOnlineAccountIdHandle, TSharedRef<FClientLobbyMemberSnapshot>>&& LobbyMemberSnapshots,
-	TMap<FOnlineAccountIdHandle, ELobbyMemberLeaveReason>&& LeaveReasons,
+	TMap<FAccountId, TSharedRef<FClientLobbyMemberSnapshot>>&& LobbyMemberSnapshots,
+	TMap<FAccountId, ELobbyMemberLeaveReason>&& LeaveReasons,
 	FLobbyEvents* LobbyEvents)
 {
 	FApplyLobbyUpdateResult Result;
@@ -140,7 +140,7 @@ FApplyLobbyUpdateResult FClientLobbyData::ApplyLobbyUpdateFromServiceSnapshot(
 	{
 		// Build list of leaving members.
 		TArray<TPair<TSharedRef<FClientLobbyMemberSnapshot>, ELobbyMemberLeaveReason>> LeavingMembers;
-		for (TPair<FOnlineAccountIdHandle, TSharedRef<FClientLobbyMemberSnapshot>>& MemberData : MemberDataStorage)
+		for (TPair<FAccountId, TSharedRef<FClientLobbyMemberSnapshot>>& MemberData : MemberDataStorage)
 		{
 			// Check if member exists in new snapshot.
 			if (LobbySnapshot.Members.Find(MemberData.Key) == nullptr)
@@ -182,7 +182,7 @@ FApplyLobbyUpdateResult FClientLobbyData::ApplyLobbyUpdateFromServiceSnapshot(
 
 	// Process member joins and attribute changes.
 	// Joining members are expected to have a member data snapshot.
-	for (TPair<FOnlineAccountIdHandle, TSharedRef<FClientLobbyMemberSnapshot>>& MemberSnapshot : LobbyMemberSnapshots)
+	for (TPair<FAccountId, TSharedRef<FClientLobbyMemberSnapshot>>& MemberSnapshot : LobbyMemberSnapshots)
 	{
 		if (LobbySnapshot.Members.Find(MemberSnapshot.Key) == nullptr)
 		{
@@ -329,7 +329,7 @@ FApplyLobbyUpdateResult FClientLobbyData::ApplyLobbyUpdateFromLocalChanges(FClie
 			{
 				LobbyEvents.OnLobbyJoined.Broadcast(FLobbyJoined{PublicData});
 
-				for (TPair<FOnlineAccountIdHandle, TSharedRef<FClientLobbyMemberSnapshot>>& MemberData : MemberDataStorage)
+				for (TPair<FAccountId, TSharedRef<FClientLobbyMemberSnapshot>>& MemberData : MemberDataStorage)
 				{
 					// The lobby creator will already be in the list during lobby creation due to applying a lobby snapshot.
 					if (Changes.MutatedMembers.Find(MemberData.Key) == nullptr)
@@ -341,7 +341,7 @@ FApplyLobbyUpdateResult FClientLobbyData::ApplyLobbyUpdateFromLocalChanges(FClie
 
 			// Create or modify state for each mutated member and notify.
 			// The lobby creator will already be in the list during lobby creation due to applying a lobby snapshot.
-			for (TPair<FOnlineAccountIdHandle, TSharedRef<FClientLobbyMemberDataChanges>>& MutatedMemberPair : Changes.MutatedMembers)
+			for (TPair<FAccountId, TSharedRef<FClientLobbyMemberDataChanges>>& MutatedMemberPair : Changes.MutatedMembers)
 			{
 				if (TSharedRef<FClientLobbyMemberSnapshot>* MemberDataPtr = MemberDataStorage.Find(MutatedMemberPair.Key))
 				{
@@ -389,7 +389,7 @@ FApplyLobbyUpdateResult FClientLobbyData::ApplyLobbyUpdateFromLocalChanges(FClie
 		// remaining lobby members and for the lobby itself.
 		if (!Changes.LeavingMembers.IsEmpty())
 		{
-			for (TPair<FOnlineAccountIdHandle, ELobbyMemberLeaveReason>& LeavingMemberPair : Changes.LeavingMembers)
+			for (TPair<FAccountId, ELobbyMemberLeaveReason>& LeavingMemberPair : Changes.LeavingMembers)
 			{
 				// Member data is expected to be added through changes before being mutated.
 				TSharedRef<FClientLobbyMemberSnapshot>* MemberDataPtr = MemberDataStorage.Find(LeavingMemberPair.Key);
diff --git a/Engine/Plugins/Online/Experimental/OnlineServices/Source/OnlineServicesCommon/Private/Online/OnlineServicesCommon.cpp b/Engine/Plugins/Online/Experimental/OnlineServices/Source/OnlineServicesCommon/Private/Online/OnlineServicesCommon.cpp
--- a/Engine/Plugins/Online/Experimental/OnlineServices/Source/OnlineServicesCommon/Private/Online/OnlineServicesCommon.cpp
+++ b/Engine/Plugins/Online/Experimental/OnlineServices/Source/OnlineServicesCommon/Private/Online/OnlineServicesCommon.cpp
@@ -174,7 +174,7 @@ FOnlineAsyncOpQueue& FOnlineServicesCommon::GetSerialQueue()
 	return SerialQueue;
 }
 
-FOnlineAsyncOpQueue& FOnlineServicesCommon::GetSerialQueue(const FOnlineAccountIdHandle& AccountId)
+FOnlineAsyncOpQueue& FOnlineServicesCommon::GetSerialQueue(const FAccountId& AccountId)
 {
 	TUniquePtr<FOnlineAsyncOpQueueSerial>* Queue = PerUserSerialQueue.Find(AccountId);
 	if (Queue == nullptr)
diff --git a/Engine/Plugins/Online/Experimental/OnlineServices/Source/OnlineServicesCommon/Private/Online/OnlineSyntaxTests.cpp b/Engine/Plugins/Online/Experimental/OnlineServices/Source/OnlineServicesCommon/Private/Online/OnlineSyntaxTests.cpp
--- a/Engine/Plugins/Online/Experimental/OnlineServices/Source/OnlineServicesCommon/Private/Online/OnlineSyntaxTests.cpp
+++ b/Engine/Plugins/Online/Experimental/OnlineServices/Source/OnlineServicesCommon/Private/Online/OnlineSyntaxTests.cpp
@@ -47,7 +47,7 @@ struct FTestUserOp
 
 	struct Params
 	{
-		FOnlineAccountIdHandle LocalUserId;
+		FAccountId LocalUserId;
 	};
 
 	struct Result
@@ -74,7 +74,7 @@ struct FJoinableTestUserOp
 
 	struct Params
 	{
-		FOnlineAccountIdHandle LocalUserId;
+		FAccountId LocalUserId;
 	};
 
 	struct Result
@@ -109,7 +109,7 @@ struct FMergeableTestUserOp
 
 	struct Params
 	{
-		FOnlineAccountIdHandle LocalUserId;
+		FAccountId LocalUserId;
 		FTestMutation Mutations;
 	};
 
diff --git a/Engine/Plugins/Online/Experimental/OnlineServices/Source/OnlineServicesCommon/Private/Online/PresenceCommon.cpp b/Engine/Plugins/Online/Experimental/OnlineServices/Source/OnlineServicesCommon/Private/Online/PresenceCommon.cpp
--- a/Engine/Plugins/Online/Experimental/OnlineServices/Source/OnlineServicesCommon/Private/Online/PresenceCommon.cpp
+++ b/Engine/Plugins/Online/Experimental/OnlineServices/Source/OnlineServicesCommon/Private/Online/PresenceCommon.cpp
@@ -41,7 +41,7 @@ public:
 
 	TFuture<FBatchQueryPresenceHelperResult> GetPromise()
 	{
-		for(FOnlineAccountIdHandle TargetUserId : Op->GetParams().TargetUserIds)
+		for(FAccountId TargetUserId : Op->GetParams().TargetUserIds)
 		{
 			FQueryPresence::Params Params;
 			Params.LocalUserId = Op->GetParams().LocalUserId;
diff --git a/Engine/Plugins/Online/Experimental/OnlineServices/Source/OnlineServicesCommon/Private/Online/SessionsCommon.cpp b/Engine/Plugins/Online/Experimental/OnlineServices/Source/OnlineServicesCommon/Private/Online/SessionsCommon.cpp
--- a/Engine/Plugins/Online/Experimental/OnlineServices/Source/OnlineServicesCommon/Private/Online/SessionsCommon.cpp
+++ b/Engine/Plugins/Online/Experimental/OnlineServices/Source/OnlineServicesCommon/Private/Online/SessionsCommon.cpp
@@ -313,9 +313,9 @@ namespace UE::Online {
 
 		FSessionSettings& SessionSettings = FoundSession->SessionSettings;
 
-		for (const TPair<FOnlineAccountIdHandle, FSessionMember>& Entry : Params.NewSessionMembers)
+		for (const TPair<FAccountId, FSessionMember>& Entry : Params.NewSessionMembers)
 		{
-			const FOnlineAccountIdHandle& SessionMemberId = Entry.Key;
+			const FAccountId& SessionMemberId = Entry.Key;
 			const FSessionMember& SessionMember = Entry.Value;
 
 			// If the user was already registered in the session, we'll update it and reserve a slot for it
@@ -378,7 +378,7 @@ namespace UE::Online {
 
 		FSessionSettings& SessionSettings = FoundSession->SessionSettings;
 
-		for (const FOnlineAccountIdHandle& SessionMemberId : Params.SessionMemberIds)
+		for (const FAccountId& SessionMemberId : Params.SessionMemberIds)
 		{
 			if (Params.bUnregisterPlayers) // If the parameters dictate it, well unregister the player as well
 			{
@@ -420,7 +420,7 @@ namespace UE::Online {
 
 		FSessionSettings& SessionSettings = FoundSession->SessionSettings;
 
-		for (const FOnlineAccountIdHandle& User : Params.TargetUsers)
+		for (const FAccountId& User : Params.TargetUsers)
 		{
 			FRegisteredPlayer& RegisteredPlayer = SessionSettings.RegisteredPlayers.FindOrAdd(User);
 
@@ -459,7 +459,7 @@ namespace UE::Online {
 
 		FSessionSettings& SessionSettings = FoundSession->SessionSettings;
 
-		for (const FOnlineAccountIdHandle& User : Params.TargetUsers)
+		for (const FAccountId& User : Params.TargetUsers)
 		{
 			if (const FRegisteredPlayer* RegisteredPlayer = SessionSettings.RegisteredPlayers.Find(User))
 			{
@@ -553,11 +553,11 @@ namespace UE::Online {
 		// User login check for all local users
 		IAuthPtr Auth = Services.GetAuthInterface();
 
-		TArray<FOnlineAccountIdHandle> LocalUserIds;
+		TArray<FAccountId> LocalUserIds;
 		LocalUserIds.Reserve(Params.LocalUsers.Num());
 		Params.LocalUsers.GenerateKeyArray(LocalUserIds);
 
-		for (const FOnlineAccountIdHandle& LocalUserId : LocalUserIds)
+		for (const FAccountId& LocalUserId : LocalUserIds)
 		{
 			if (!Auth->IsLoggedIn(LocalUserId))
 			{
@@ -652,11 +652,11 @@ namespace UE::Online {
 		// User login check for all local users
 		IAuthPtr Auth = Services.GetAuthInterface();
 
-		TArray<FOnlineAccountIdHandle> LocalUserIds;
+		TArray<FAccountId> LocalUserIds;
 		LocalUserIds.Reserve(Params.LocalUsers.Num());
 		Params.LocalUsers.GenerateKeyArray(LocalUserIds);
 
-		for (const FOnlineAccountIdHandle& LocalUserId : LocalUserIds)
+		for (const FAccountId& LocalUserId : LocalUserIds)
 		{
 			if (!Auth->IsLoggedIn(LocalUserId))
 			{
@@ -717,11 +717,11 @@ namespace UE::Online {
 		// User login check for all local users
 		IAuthPtr Auth = Services.GetAuthInterface();
 
-		TArray<FOnlineAccountIdHandle> LocalUserIds;
+		TArray<FAccountId> LocalUserIds;
 		LocalUserIds.Reserve(Params.LocalUsers.Num());
 		Params.LocalUsers.GenerateKeyArray(LocalUserIds);
 
-		for (const FOnlineAccountIdHandle& LocalUserId : LocalUserIds)
+		for (const FAccountId& LocalUserId : LocalUserIds)
 		{
 			if(!Auth->IsLoggedIn(LocalUserId))
 			{
@@ -751,7 +751,7 @@ namespace UE::Online {
 			return Errors::InvalidParams();
 		}
 
-		for (const FOnlineAccountIdHandle& LocalUserId : LocalUserIds)
+		for (const FAccountId& LocalUserId : LocalUserIds)
 		{
 			if (SessionSettings.SessionMembers.Contains(LocalUserId))
 			{
@@ -770,7 +770,7 @@ namespace UE::Online {
 
 		if (!SessionSettings.bAllowUnregisteredPlayers)
 		{
-			for (const FOnlineAccountIdHandle& LocalUserId : LocalUserIds)
+			for (const FAccountId& LocalUserId : LocalUserIds)
 			{
 				if (!SessionSettings.RegisteredPlayers.Contains(LocalUserId))
 				{
@@ -837,7 +837,7 @@ namespace UE::Online {
 		IAuthPtr Auth = Services.GetAuthInterface();
 
 		// User login check for all local users
-		for (const FOnlineAccountIdHandle& LocalUserId : Params.LocalUsers)
+		for (const FAccountId& LocalUserId : Params.LocalUsers)
 		{
 			if (!Auth->IsLoggedIn(LocalUserId))
 			{
diff --git a/Engine/Plugins/Online/Experimental/OnlineServices/Source/OnlineServicesCommon/Public/Online/AchievementsCommon.h b/Engine/Plugins/Online/Experimental/OnlineServices/Source/OnlineServicesCommon/Public/Online/AchievementsCommon.h
--- a/Engine/Plugins/Online/Experimental/OnlineServices/Source/OnlineServicesCommon/Public/Online/AchievementsCommon.h
+++ b/Engine/Plugins/Online/Experimental/OnlineServices/Source/OnlineServicesCommon/Public/Online/AchievementsCommon.h
@@ -51,9 +51,9 @@ protected:
 	TOnlineEventCallable<void(const FAchievementStateUpdated&)> OnAchievementStateUpdatedEvent;
 
 	void UnlockAchievementsByStats(const FStatsUpdated& StatsUpdated);
-	void ExecuteUnlockRulesRelatedToStat(const FOnlineAccountIdHandle& UserId, const FString& StatName, const TMap<FString, FStatValue>& Stats, TArray<FString>& OutAchievementsToUnlock);
+	void ExecuteUnlockRulesRelatedToStat(const FAccountId& UserId, const FString& StatName, const TMap<FString, FStatValue>& Stats, TArray<FString>& OutAchievementsToUnlock);
 	bool MeetUnlockCondition(FAchievementUnlockRule AchievementUnlockRule, const TMap<FString, FStatValue>& Stats);
-	bool IsUnlocked(const FOnlineAccountIdHandle& UserId, const FString& AchievementName) const;
+	bool IsUnlocked(const FAccountId& UserId, const FString& AchievementName) const;
 
 	TArray<FAchievementUnlockRule> AchievementUnlockRules;
 	FOnlineEventDelegateHandle StatEventHandle;
diff --git a/Engine/Plugins/Online/Experimental/OnlineServices/Source/OnlineServicesCommon/Public/Online/AuthCommon.h b/Engine/Plugins/Online/Experimental/OnlineServices/Source/OnlineServicesCommon/Public/Online/AuthCommon.h
--- a/Engine/Plugins/Online/Experimental/OnlineServices/Source/OnlineServicesCommon/Public/Online/AuthCommon.h
+++ b/Engine/Plugins/Online/Experimental/OnlineServices/Source/OnlineServicesCommon/Public/Online/AuthCommon.h
@@ -15,7 +15,7 @@ public:
 	virtual ~FAccountInfoRegistry() = default;
 
 	TSharedPtr<FAccountInfo> Find(FPlatformUserId PlatformUserId) const;
-	TSharedPtr<FAccountInfo> Find(FOnlineAccountIdHandle AccountIdHandle) const;
+	TSharedPtr<FAccountInfo> Find(FAccountId AccountIdHandle) const;
 
 	TArray<TSharedRef<FAccountInfo>> GetAllAccountInfo(TFunction<bool(const TSharedRef<FAccountInfo>&)> Predicate) const;
 
@@ -27,7 +27,7 @@ protected:
 
 private:
 	TMap<FPlatformUserId, TSharedRef<FAccountInfo>> AuthDataByPlatformUserId;
-	TMap<FOnlineAccountIdHandle, TSharedRef<FAccountInfo>> AuthDataByOnlineAccountIdHandle;
+	TMap<FAccountId, TSharedRef<FAccountInfo>> AuthDataByOnlineAccountIdHandle;
 };
 
 class ONLINESERVICESCOMMON_API FAuthCommon : public TOnlineComponent<IAuth>
@@ -56,7 +56,7 @@ public:
 	virtual TOnlineEvent<void(const FAuthLoginStatusChanged&)> OnLoginStatusChanged() override;
 	virtual TOnlineEvent<void(const FAuthPendingAuthExpiration&)> OnPendingAuthExpiration() override;
 	virtual TOnlineEvent<void(const FAuthAccountAttributesChanged&)> OnAccountAttributesChanged() override;
-	virtual bool IsLoggedIn(const FOnlineAccountIdHandle& AccountId) const override;
+	virtual bool IsLoggedIn(const FAccountId& AccountId) const override;
 
 protected:
 	virtual const FAccountInfoRegistry& GetAccountInfoRegistry() const = 0;
diff --git a/Engine/Plugins/Online/Experimental/OnlineServices/Source/OnlineServicesCommon/Public/Online/LeaderboardsCommon.h b/Engine/Plugins/Online/Experimental/OnlineServices/Source/OnlineServicesCommon/Public/Online/LeaderboardsCommon.h
--- a/Engine/Plugins/Online/Experimental/OnlineServices/Source/OnlineServicesCommon/Public/Online/LeaderboardsCommon.h
+++ b/Engine/Plugins/Online/Experimental/OnlineServices/Source/OnlineServicesCommon/Public/Online/LeaderboardsCommon.h
@@ -47,7 +47,7 @@ struct FWriteLeaderboardScores
 
 	struct Params
 	{
-		FOnlineAccountIdHandle LocalUserId;
+		FAccountId LocalUserId;
 		FString BoardName;
 		uint64 Score;
 	};
diff --git a/Engine/Plugins/Online/Experimental/OnlineServices/Source/OnlineServicesCommon/Public/Online/LobbiesCommonTypes.h b/Engine/Plugins/Online/Experimental/OnlineServices/Source/OnlineServicesCommon/Public/Online/LobbiesCommonTypes.h
--- a/Engine/Plugins/Online/Experimental/OnlineServices/Source/OnlineServicesCommon/Public/Online/LobbiesCommonTypes.h
+++ b/Engine/Plugins/Online/Experimental/OnlineServices/Source/OnlineServicesCommon/Public/Online/LobbiesCommonTypes.h
@@ -195,7 +195,7 @@ struct FClientLobbyDataChanges
 	TOptional<ELobbyJoinPolicy> JoinPolicy;
 
 	/** Setting for lobby ownership change. */
-	TOptional<FOnlineAccountIdHandle> OwnerAccountId;
+	TOptional<FAccountId> OwnerAccountId;
 
 	/** Setting for lobby schema change. */
 	TOptional<FLobbySchemaId> LobbySchema;
@@ -207,10 +207,10 @@ struct FClientLobbyDataChanges
 	TSet<FLobbyAttributeId> ClearedAttributes;
 
 	/** Members to be added or changed. */
-	TMap<FOnlineAccountIdHandle, TSharedRef<FClientLobbyMemberDataChanges>> MutatedMembers;
+	TMap<FAccountId, TSharedRef<FClientLobbyMemberDataChanges>> MutatedMembers;
 
 	/** Members to be removed. */
-	TMap<FOnlineAccountIdHandle, ELobbyMemberLeaveReason> LeavingMembers;
+	TMap<FAccountId, ELobbyMemberLeaveReason> LeavingMembers;
 };
 
 /**
@@ -220,12 +220,12 @@ struct FClientLobbyDataChanges
 */
 struct FClientLobbySnapshot
 {
-	FOnlineAccountIdHandle OwnerAccountId;
+	FAccountId OwnerAccountId;
 	FName SchemaName;
 	int32 MaxMembers;
 	ELobbyJoinPolicy JoinPolicy;
 	TMap<FLobbyAttributeId, FLobbyVariant> Attributes;
-	TSet<FOnlineAccountIdHandle> Members;
+	TSet<FAccountId> Members;
 };
 
 struct FClientLobbyMemberSnapshot : public FLobbyMember
@@ -235,7 +235,7 @@ struct FClientLobbyMemberSnapshot : public FLobbyMember
 
 struct FApplyLobbyUpdateResult
 {
-	TArray<FOnlineAccountIdHandle> LeavingLocalMembers;
+	TArray<FAccountId> LeavingLocalMembers;
 };
 
 /** Lobby data as seen by the client. */
@@ -247,7 +247,7 @@ public:
 	TSharedRef<const FLobby> GetPublicDataPtr() const { return PublicData; }
 	const FLobby& GetPublicData() const { return *PublicData; }
 
-	TSharedPtr<const FClientLobbyMemberSnapshot> GetMemberData(FOnlineAccountIdHandle MemberAccountId) const;
+	TSharedPtr<const FClientLobbyMemberSnapshot> GetMemberData(FAccountId MemberAccountId) const;
 
 	/**
 	 * Apply updated lobby data and generate changes.
@@ -256,8 +256,8 @@ public:
 	 */
 	FApplyLobbyUpdateResult ApplyLobbyUpdateFromServiceSnapshot(
 		FClientLobbySnapshot&& LobbySnapshot,
-		TMap<FOnlineAccountIdHandle, TSharedRef<FClientLobbyMemberSnapshot>>&& LobbyMemberSnapshots,
-		TMap<FOnlineAccountIdHandle, ELobbyMemberLeaveReason>&& LeaveReasons = TMap<FOnlineAccountIdHandle, ELobbyMemberLeaveReason>(),
+		TMap<FAccountId, TSharedRef<FClientLobbyMemberSnapshot>>&& LobbyMemberSnapshots,
+		TMap<FAccountId, ELobbyMemberLeaveReason>&& LeaveReasons = TMap<FAccountId, ELobbyMemberLeaveReason>(),
 		FLobbyEvents* LobbyEvents = nullptr);
 
 	/**
@@ -286,13 +286,13 @@ private:
 	TSharedRef<FLobby> PublicData;
 
 	/** Mutable lobby member data storage. */
-	TMap<FOnlineAccountIdHandle, TSharedRef<FClientLobbyMemberSnapshot>> MemberDataStorage;
+	TMap<FAccountId, TSharedRef<FClientLobbyMemberSnapshot>> MemberDataStorage;
 
 	/**
 	 * Keep track of which members are local to the client.
 	 * When all local members have been removed all members will be removed.
 	 */
-	TSet<FOnlineAccountIdHandle> LocalMembers;
+	TSet<FAccountId> LocalMembers;
 };
 
 #if LOBBIES_FUNCTIONAL_TEST_ENABLED
diff --git a/Engine/Plugins/Online/Experimental/OnlineServices/Source/OnlineServicesCommon/Public/Online/OnlineAsyncOpCache.h b/Engine/Plugins/Online/Experimental/OnlineServices/Source/OnlineServicesCommon/Public/Online/OnlineAsyncOpCache.h
--- a/Engine/Plugins/Online/Experimental/OnlineServices/Source/OnlineServicesCommon/Public/Online/OnlineAsyncOpCache.h
+++ b/Engine/Plugins/Online/Experimental/OnlineServices/Source/OnlineServicesCommon/Public/Online/OnlineAsyncOpCache.h
@@ -494,7 +494,7 @@ private:
 	friend uint32 GetTypeHash(const FWrappedOperationKey& Key);
 
 	TMap<FWrappedOperationKey, TUniquePtr<IWrappedOperation>> Operations;
-	TMap<FOnlineAccountIdHandle, TMap<FWrappedOperationKey, TUniquePtr<IWrappedOperation>>> UserOperations;
+	TMap<FAccountId, TMap<FWrappedOperationKey, TUniquePtr<IWrappedOperation>>> UserOperations;
 };
 
 inline uint32 GetTypeHash(const FOnlineAsyncOpCache::FWrappedOperationKey& Key)
diff --git a/Engine/Plugins/Online/Experimental/OnlineServices/Source/OnlineServicesCommon/Public/Online/OnlineComponent.h b/Engine/Plugins/Online/Experimental/OnlineServices/Source/OnlineServicesCommon/Public/Online/OnlineComponent.h
--- a/Engine/Plugins/Online/Experimental/OnlineServices/Source/OnlineServicesCommon/Public/Online/OnlineComponent.h
+++ b/Engine/Plugins/Online/Experimental/OnlineServices/Source/OnlineServicesCommon/Public/Online/OnlineComponent.h
@@ -87,7 +87,7 @@ public:
 	}
 
 	/* Queues for executing per-user tasks in serial */
-	FOnlineAsyncOpQueue& GetSerialQueue(const FOnlineAccountIdHandle& AccountId)
+	FOnlineAsyncOpQueue& GetSerialQueue(const FAccountId& AccountId)
 	{
 		TUniquePtr<FOnlineAsyncOpQueueSerial>* Queue = PerUserSerialQueue.Find(AccountId);
 		if (Queue == nullptr)
@@ -210,7 +210,7 @@ private:
 	TMap<FString, TUniquePtr<IOnlineExecHandler>> ExecCommands;
 
 	FOnlineAsyncOpQueueSerial SerialQueue;
-	TMap<FOnlineAccountIdHandle, TUniquePtr<FOnlineAsyncOpQueueSerial>> PerUserSerialQueue;
+	TMap<FAccountId, TUniquePtr<FOnlineAsyncOpQueueSerial>> PerUserSerialQueue;
 };
 
 /* UE::Online */ }
diff --git a/Engine/Plugins/Online/Experimental/OnlineServices/Source/OnlineServicesCommon/Public/Online/OnlineExecHandler.h b/Engine/Plugins/Online/Experimental/OnlineServices/Source/OnlineServicesCommon/Public/Online/OnlineExecHandler.h
--- a/Engine/Plugins/Online/Experimental/OnlineServices/Source/OnlineServicesCommon/Public/Online/OnlineExecHandler.h
+++ b/Engine/Plugins/Online/Experimental/OnlineServices/Source/OnlineServicesCommon/Public/Online/OnlineExecHandler.h
@@ -22,7 +22,7 @@
 	Next argument is the interface- i.e. "Auth", "Presence"
 	Next argument is the name of the function being called, i.e. "QueryPresence"
 	Arguments after that are the function's params in declared order. Special types are as follows:
-		for FOnlineAccountIdHandle, either pass in 0-9 to grab the nth local user, or pass in the output from a handle's tostring that looks like "epic:3"
+		for FAccountId, either pass in 0-9 to grab the nth local user, or pass in the output from a handle's tostring that looks like "epic:3"
 		for TSharedPtr, simply type in the parameter as normal
 			- TSharedRef *technically* is implemented however will crash because you cannot default construct shared refs. At the moment, we recommend all console commands with shared ptr parameters to use tsharedptr instead of tsharedref.
 		for TOptional, pass in "null" and the optional will be unset
@@ -32,7 +32,7 @@
 			- i32 for int32
 			- b for bool
 			- s for string
-			- u for user (FOnlineAccountIdHandle)
+			- u for user (FAccountId)
 			- d for double
 			- f for float
 			- e for enum
@@ -222,18 +222,18 @@ public:
 };
 
 template< typename VariantObject >
-struct TOnlineVariantVisitInfo<FOnlineAccountIdHandle, VariantObject>
+struct TOnlineVariantVisitInfo<FAccountId, VariantObject>
 {
 public:
 	inline static TArray<FString> Prefixes = { TEXT("u"),TEXT("user") };
 
 	static bool Assign(FString& VariantValue, VariantObject& Variant, IOnlineServices* Services)
 	{
-		FOnlineAccountIdHandle Handle;
+		FAccountId Handle;
 		const TCHAR* ValueTCHAR = *VariantValue;
 		if (ParseOnlineExecParams<TOnlineIdHandle<OnlineIdHandleTags::FAccount>>(ValueTCHAR, Handle, Services))
 		{
-			Variant.template Set<FOnlineAccountIdHandle>(Handle);
+			Variant.template Set<FAccountId>(Handle);
 			return true;
 		}
 		return false;
diff --git a/Engine/Plugins/Online/Experimental/OnlineServices/Source/OnlineServicesCommon/Public/Online/OnlineServicesCommon.h b/Engine/Plugins/Online/Experimental/OnlineServices/Source/OnlineServicesCommon/Public/Online/OnlineServicesCommon.h
--- a/Engine/Plugins/Online/Experimental/OnlineServices/Source/OnlineServicesCommon/Public/Online/OnlineServicesCommon.h
+++ b/Engine/Plugins/Online/Experimental/OnlineServices/Source/OnlineServicesCommon/Public/Online/OnlineServicesCommon.h
@@ -306,7 +306,7 @@ public:
 	FOnlineAsyncOpQueue& GetSerialQueue();
 
 	/* Queues for executing per-user tasks in serial */
-	FOnlineAsyncOpQueue& GetSerialQueue(const FOnlineAccountIdHandle& AccountId);
+	FOnlineAsyncOpQueue& GetSerialQueue(const FAccountId& AccountId);
 	
 	void RegisterExecHandler(const FString& Name, TUniquePtr<IOnlineExecHandler>&& Handler);
 
@@ -330,7 +330,7 @@ protected:
 
 	FOnlineAsyncOpQueueParallel ParallelQueue;
 	FOnlineAsyncOpQueueSerial SerialQueue;
-	TMap<FOnlineAccountIdHandle, TUniquePtr<FOnlineAsyncOpQueueSerial>> PerUserSerialQueue;
+	TMap<FAccountId, TUniquePtr<FOnlineAsyncOpQueueSerial>> PerUserSerialQueue;
 };
 
 /* UE::Online */ }
diff --git a/Engine/Plugins/Online/Experimental/OnlineServices/Source/OnlineServicesCommon/Public/Online/SessionsCommon.h b/Engine/Plugins/Online/Experimental/OnlineServices/Source/OnlineServicesCommon/Public/Online/SessionsCommon.h
--- a/Engine/Plugins/Online/Experimental/OnlineServices/Source/OnlineServicesCommon/Public/Online/SessionsCommon.h
+++ b/Engine/Plugins/Online/Experimental/OnlineServices/Source/OnlineServicesCommon/Public/Online/SessionsCommon.h
@@ -214,9 +214,9 @@ protected:
 	/** Map of sessions that local users are part of, indexed by their local name */
 	TMap<FName, TSharedRef<FSession>> LocalSessionsByName;
 	/** Cache for the last set of session search results, mapped per user */
-	TMap<FOnlineAccountIdHandle, TMap<FOnlineSessionIdHandle, TSharedRef<FSession>>> SessionSearchResultsUserMap;
+	TMap<FAccountId, TMap<FOnlineSessionIdHandle, TSharedRef<FSession>>> SessionSearchResultsUserMap;
 	/** Cache for received session invites, mapped per user */
-	TMap<FOnlineAccountIdHandle, TMap<FOnlineSessionInviteIdHandle, TSharedRef<FSessionInvite>>> SessionInvitesUserMap;
+	TMap<FAccountId, TMap<FOnlineSessionInviteIdHandle, TSharedRef<FSessionInvite>>> SessionInvitesUserMap;
 
 	TSharedPtr<FFindSessions::Result> CurrentSessionSearch;
 	TSharedPtr<TOnlineAsyncOp<FFindSessions>> CurrentSessionSearchHandle;
diff --git a/Engine/Plugins/Online/Experimental/OnlineServices/Source/OnlineServicesCommon/Public/Online/StatsCommon.h b/Engine/Plugins/Online/Experimental/OnlineServices/Source/OnlineServicesCommon/Public/Online/StatsCommon.h
--- a/Engine/Plugins/Online/Experimental/OnlineServices/Source/OnlineServicesCommon/Public/Online/StatsCommon.h
+++ b/Engine/Plugins/Online/Experimental/OnlineServices/Source/OnlineServicesCommon/Public/Online/StatsCommon.h
@@ -50,7 +50,7 @@ struct FStatDefinition
 
 struct FFindUserStatsByAccountId
 {
-	FFindUserStatsByAccountId(const FOnlineAccountIdHandle& InUserId)
+	FFindUserStatsByAccountId(const FAccountId& InUserId)
 		: UserId(InUserId)
 	{
 	}
@@ -60,7 +60,7 @@ struct FFindUserStatsByAccountId
 		return UserStats.UserId == UserId;
 	}
 
-	FOnlineAccountIdHandle UserId;
+	FAccountId UserId;
 };
 
 class ONLINESERVICESCOMMON_API FStatsCommon : public TOnlineComponent<IStats>
diff --git a/Engine/Plugins/Online/Experimental/OnlineServices/Source/OnlineServicesCommonEngineUtils/Private/Online/SessionsLAN.cpp b/Engine/Plugins/Online/Experimental/OnlineServices/Source/OnlineServicesCommonEngineUtils/Private/Online/SessionsLAN.cpp
--- a/Engine/Plugins/Online/Experimental/OnlineServices/Source/OnlineServicesCommonEngineUtils/Private/Online/SessionsLAN.cpp
+++ b/Engine/Plugins/Online/Experimental/OnlineServices/Source/OnlineServicesCommonEngineUtils/Private/Online/SessionsLAN.cpp
@@ -298,7 +298,7 @@ TOnlineAsyncOpHandle<FJoinSession> FSessionsLAN::JoinSession(FJoinSession::Param
 
 		Op.SetResult(FJoinSession::Result{ NewSessionRef });
 
-		TArray<FOnlineAccountIdHandle> LocalUserIds;
+		TArray<FAccountId> LocalUserIds;
 		LocalUserIds.Reserve(OpParams.LocalUsers.Num());
 		OpParams.LocalUsers.GenerateKeyArray(LocalUserIds);
 
diff --git a/Engine/Plugins/Online/Experimental/OnlineServices/Source/OnlineServicesInterface/Private/Online/Sessions.cpp b/Engine/Plugins/Online/Experimental/OnlineServices/Source/OnlineServicesInterface/Private/Online/Sessions.cpp
--- a/Engine/Plugins/Online/Experimental/OnlineServices/Source/OnlineServicesInterface/Private/Online/Sessions.cpp
+++ b/Engine/Plugins/Online/Experimental/OnlineServices/Source/OnlineServicesInterface/Private/Online/Sessions.cpp
@@ -112,19 +112,19 @@ FSessionSettings& FSessionSettings::operator+=(const FSessionSettingsUpdate& Upd
 
 	CustomSettings.Append(UpdatedSettings.UpdatedCustomSettings);
 
-	for (const FOnlineAccountIdHandle& Key : UpdatedSettings.RemovedRegisteredPlayers)
+	for (const FAccountId& Key : UpdatedSettings.RemovedRegisteredPlayers)
 	{
 		RegisteredPlayers.Remove(Key);
 	}
 
 	RegisteredPlayers.Append(UpdatedSettings.UpdatedRegisteredPlayers);
 
-	for (const FOnlineAccountIdHandle& Key : UpdatedSettings.RemovedSessionMembers)
+	for (const FAccountId& Key : UpdatedSettings.RemovedSessionMembers)
 	{
 		SessionMembers.Remove(Key);
 	}
 
-	for (const TPair<FOnlineAccountIdHandle, FSessionMemberUpdate>& Entry : UpdatedSettings.UpdatedSessionMembers)
+	for (const TPair<FAccountId, FSessionMemberUpdate>& Entry : UpdatedSettings.UpdatedSessionMembers)
 	{
 		if (FSessionMember* SessionMember = SessionMembers.Find(Entry.Key))
 		{
@@ -199,7 +199,7 @@ FSessionSettingsUpdate& FSessionSettingsUpdate::operator+=(FSessionSettingsUpdat
 		RemovedCustomSettings.AddUnique(MoveTemp(Key));
 	}
 
-	for (TPair<FOnlineAccountIdHandle, FSessionMemberUpdate>& UpdatedSessionMember : UpdatedValue.UpdatedSessionMembers)
+	for (TPair<FAccountId, FSessionMemberUpdate>& UpdatedSessionMember : UpdatedValue.UpdatedSessionMembers)
 	{
 		// If an update adds a modification to a member that had previously been marked for removal, we'll keep the latest change
 		RemovedSessionMembers.Remove(UpdatedSessionMember.Key);
@@ -208,21 +208,21 @@ FSessionSettingsUpdate& FSessionSettingsUpdate::operator+=(FSessionSettingsUpdat
 		MemberUpdate += MoveTemp(UpdatedSessionMember.Value);
 	}
 
-	for (FOnlineAccountIdHandle& Key : UpdatedValue.RemovedSessionMembers)
+	for (FAccountId& Key : UpdatedValue.RemovedSessionMembers)
 	{
 		// If an update removes a member that had previously been modified, we'll keep the latest change
 		UpdatedSessionMembers.Remove(Key);
 		RemovedSessionMembers.AddUnique(MoveTemp(Key));
 	}
 
-	for (const TPair<FOnlineAccountIdHandle, FRegisteredPlayer>& UpdatedRegisteredPlayer : UpdatedValue.UpdatedRegisteredPlayers)
+	for (const TPair<FAccountId, FRegisteredPlayer>& UpdatedRegisteredPlayer : UpdatedValue.UpdatedRegisteredPlayers)
 	{
 		// If an update adds a modification to a registered player that had previously been marked for removal, we'll keep the latest change
 		RemovedRegisteredPlayers.Remove(UpdatedRegisteredPlayer.Key);
 	}
 	UpdatedRegisteredPlayers.Append(MoveTemp(UpdatedValue.UpdatedRegisteredPlayers));
 
-	for (FOnlineAccountIdHandle& Key : UpdatedValue.RemovedRegisteredPlayers)
+	for (FAccountId& Key : UpdatedValue.RemovedRegisteredPlayers)
 	{
 		// If an update removes a registered player that had previously been modified, we'll keep the latest change
 		UpdatedRegisteredPlayers.Remove(Key);
diff --git a/Engine/Plugins/Online/Experimental/OnlineServices/Source/OnlineServicesInterface/Public/Online/Achievements.h b/Engine/Plugins/Online/Experimental/OnlineServices/Source/OnlineServicesInterface/Public/Online/Achievements.h
--- a/Engine/Plugins/Online/Experimental/OnlineServices/Source/OnlineServicesInterface/Public/Online/Achievements.h
+++ b/Engine/Plugins/Online/Experimental/OnlineServices/Source/OnlineServicesInterface/Public/Online/Achievements.h
@@ -16,7 +16,7 @@ struct FQueryAchievementDefinitions
 	struct Params
 	{
 		/** Local user performing the operation */
-		FOnlineAccountIdHandle LocalUserId;
+		FAccountId LocalUserId;
 	};
 
 	struct Result
@@ -31,7 +31,7 @@ struct FGetAchievementIds
 	struct Params
 	{
 		/** Local user performing the operation */
-		FOnlineAccountIdHandle LocalUserId;
+		FAccountId LocalUserId;
 	};
 
 	struct Result
@@ -84,7 +84,7 @@ struct FGetAchievementDefinition
 	struct Params
 	{
 		/** Local user performing the operation */
-		FOnlineAccountIdHandle LocalUserId;
+		FAccountId LocalUserId;
 		/** Achievement to get definition of */
 		FString AchievementId;
 	};
@@ -103,7 +103,7 @@ struct FQueryAchievementStates
 	struct Params
 	{
 		/** Local user performing the operation */
-		FOnlineAccountIdHandle LocalUserId;
+		FAccountId LocalUserId;
 	};
 
 	struct Result
@@ -130,7 +130,7 @@ struct FGetAchievementState
 	struct Params
 	{
 		/** Local user performing the operation */
-		FOnlineAccountIdHandle LocalUserId;
+		FAccountId LocalUserId;
 		/** Achievement to get state of */
 		FString AchievementId;
 	};
@@ -149,7 +149,7 @@ struct FUnlockAchievements
 	struct Params
 	{
 		/** Local user performing the operation */
-		FOnlineAccountIdHandle LocalUserId;
+		FAccountId LocalUserId;
 		/** Achievements to unlock */
 		TArray<FString> AchievementIds;
 	};
@@ -166,7 +166,7 @@ struct FDisplayAchievementUI
 	struct Params
 	{
 		/** Local user performing the operation */
-		FOnlineAccountIdHandle LocalUserId;
+		FAccountId LocalUserId;
 		/** Achievement to display */
 		FString AchievementId;
 	};
@@ -180,7 +180,7 @@ struct FDisplayAchievementUI
 struct FAchievementStateUpdated
 {
 	/** User whose achievement states have updated */
-	FOnlineAccountIdHandle LocalUserId;
+	FAccountId LocalUserId;
 	/** Achievements which have updated */
 	TArray<FString> AchievementIds;
 };
diff --git a/Engine/Plugins/Online/Experimental/OnlineServices/Source/OnlineServicesInterface/Public/Online/Auth.h b/Engine/Plugins/Online/Experimental/OnlineServices/Source/OnlineServicesInterface/Public/Online/Auth.h
--- a/Engine/Plugins/Online/Experimental/OnlineServices/Source/OnlineServicesInterface/Public/Online/Auth.h
+++ b/Engine/Plugins/Online/Experimental/OnlineServices/Source/OnlineServicesInterface/Public/Online/Auth.h
@@ -109,7 +109,7 @@ ONLINESERVICESINTERFACE_API void LexFromString(EExternalAuthTokenMethod& OutMeth
 struct FAccountInfo
 {
 	/** The account id for the user which represents the user's online platform account. */
-	FOnlineAccountIdHandle AccountId;
+	FAccountId AccountId;
 	/** The platform user id associated with the online user. */
 	FPlatformUserId PlatformUserId = PLATFORMUSERID_NONE;
 	/** Login status */
@@ -142,7 +142,7 @@ using FCredentialsToken = TVariant<FString, FExternalAuthToken>;
 struct FVerifiedAuthSession
 {
 	FOnlineVerifiedAuthSessionIdHandle SessionId;
-	FOnlineAccountIdHandle RemoteAccountId;
+	FAccountId RemoteAccountId;
 	double CreationTime;
 };
 
@@ -182,7 +182,7 @@ struct FAuthLogout
 	struct Params
 	{
 		/** The online account id of the Local User making the request. */
-		FOnlineAccountIdHandle LocalUserId;
+		FAccountId LocalUserId;
 		/** Whether to remove persistent credentials on logout. */
 		bool bDestroyAuth = false;
 	};
@@ -199,7 +199,7 @@ struct FAuthModifyAccountAttributes
 	struct Params
 	{
 		/** The online account id of the Local User making the request. */
-		FOnlineAccountIdHandle LocalUserId;
+		FAccountId LocalUserId;
 		/** New or changed attributes. */
 		TMap<FSchemaAttributeId, FSchemaVariant> MutatedAttributes;
 		/** Attributes to be cleared. */
@@ -218,7 +218,7 @@ struct FAuthQueryExternalServerAuthTicket
 	struct Params
 	{
 		/** The online account id of the Local User making the request. */
-		FOnlineAccountIdHandle LocalUserId;
+		FAccountId LocalUserId;
 	};
 
 	struct Result
@@ -234,7 +234,7 @@ struct FAuthQueryExternalAuthToken
 	struct Params
 	{
 		/** The online account id of the Local User making the request. */
-		FOnlineAccountIdHandle LocalUserId;
+		FAccountId LocalUserId;
 		/** The method of external auth to provide. */
 		EExternalAuthTokenMethod Method = EExternalAuthTokenMethod::Primary;
 	};
@@ -252,7 +252,7 @@ struct FAuthQueryVerifiedAuthTicket
 	struct Params
 	{
 		/** The online account id of the Local User making the request. */
-		FOnlineAccountIdHandle LocalUserId;
+		FAccountId LocalUserId;
 		/** The intended purpose of the auth ticket. */
 		ERemoteAuthTicketAudience Audience = ERemoteAuthTicketAudience::Peer;
 	};
@@ -273,7 +273,7 @@ struct FAuthCancelVerifiedAuthTicket
 	struct Params
 	{
 		/** The online account id of the Local User making the request. */
-		FOnlineAccountIdHandle LocalUserId;
+		FAccountId LocalUserId;
 		/** Local ticket id used to reference the verified auth ticket. */
 		FOnlineVerifiedAuthTicketIdHandle VerifiedAuthTicketId;
 	};
@@ -290,7 +290,7 @@ struct FAuthBeginVerifiedAuthSession
 	struct Params
 	{
 		/** The remote user for which to start the verified auth session. */
-		FOnlineAccountIdHandle RemoteAccountId;
+		FAccountId RemoteAccountId;
 		/** The ticket used to verify the identity of the user. */
 		FVerifiedAuthTicket Ticket;
 	};
@@ -322,7 +322,7 @@ struct FAuthGetLocalOnlineUserByOnlineAccountId
 	struct Params
 	{
 		/** Account Id of the Local User making the request */
-		FOnlineAccountIdHandle LocalUserId;
+		FAccountId LocalUserId;
 	};
 
 	struct Result
@@ -516,7 +516,7 @@ public:
 	/**
 	 * Helper for querying the login status of a local user.
 	 */
-	virtual bool IsLoggedIn(const FOnlineAccountIdHandle& AccountId) const = 0;
+	virtual bool IsLoggedIn(const FAccountId& AccountId) const = 0;
 };
 
 namespace Meta {
diff --git a/Engine/Plugins/Online/Experimental/OnlineServices/Source/OnlineServicesInterface/Public/Online/Commerce.h b/Engine/Plugins/Online/Experimental/OnlineServices/Source/OnlineServicesInterface/Public/Online/Commerce.h
--- a/Engine/Plugins/Online/Experimental/OnlineServices/Source/OnlineServicesInterface/Public/Online/Commerce.h
+++ b/Engine/Plugins/Online/Experimental/OnlineServices/Source/OnlineServicesInterface/Public/Online/Commerce.h
@@ -66,7 +66,7 @@ struct FCommerceQueryOffers
 	struct Params
 	{
 		/* Local user */
-		FOnlineAccountIdHandle LocalUserId;
+		FAccountId LocalUserId;
 	};
 
 	struct Result
@@ -83,7 +83,7 @@ struct FCommerceQueryOffersById
 	struct Params
 	{
 		/* Local user */
-		FOnlineAccountIdHandle LocalUserId;
+		FAccountId LocalUserId;
 		/* List of offer ids */
 		TArray<FOfferId> OfferIds;
 	};
@@ -103,7 +103,7 @@ struct FCommerceGetOffers
 	struct Params
 	{
 		/* Local user */
-		FOnlineAccountIdHandle LocalUserId;
+		FAccountId LocalUserId;
 	};
 
 	struct Result
@@ -121,7 +121,7 @@ struct FCommerceGetOffersById
 	struct Params
 	{
 		/* Local user */
-		FOnlineAccountIdHandle LocalUserId;
+		FAccountId LocalUserId;
 		/* List of offer ids */
 		TArray<FOfferId> OfferIds;
 	};
@@ -142,7 +142,7 @@ struct FCommerceShowStoreUI
 	struct Params
 	{
 		/* Local user */
-		FOnlineAccountIdHandle LocalUserId;
+		FAccountId LocalUserId;
 	};
 
 	struct Result
@@ -167,7 +167,7 @@ struct FCommerceCheckout
 	struct Params
 	{
 		/* Local user */
-		FOnlineAccountIdHandle LocalUserId;
+		FAccountId LocalUserId;
 		/* Offers to purchase */
 		TArray<FPurchaseOffer> Offers;
 	};
@@ -206,7 +206,7 @@ struct FCommerceQueryTransactionEntitlements
 	struct Params
 	{
 		/* Local user */
-		FOnlineAccountIdHandle LocalUserId;
+		FAccountId LocalUserId;
 		/* Transaction id for a completed purchase */
 		FString TransactionId;
 	};
@@ -226,7 +226,7 @@ struct FCommerceQueryEntitlements
 	struct Params
 	{
 		/* Local user */
-		FOnlineAccountIdHandle LocalUserId;
+		FAccountId LocalUserId;
 		/* Whether or not to fetch redeemed entitlements */
 		bool bIncludeRedeemed;
 	};
@@ -245,7 +245,7 @@ struct FCommerceGetEntitlements
 	struct Params
 	{
 		/* Local user */
-		FOnlineAccountIdHandle LocalUserId;
+		FAccountId LocalUserId;
 	};
 
 	struct Result
@@ -263,7 +263,7 @@ struct FCommerceRedeemEntitlement
 	struct Params
 	{
 		/* Local user */
-		FOnlineAccountIdHandle LocalUserId;
+		FAccountId LocalUserId;
 		/* Entitlement id to mark as redeemed */
 		FEntitlementId EntitlementId;
 		/* Quantity to consume */
@@ -284,7 +284,7 @@ struct FCommerceRetrieveS2SToken
 	struct Params
 	{
 		/* Local user */
-		FOnlineAccountIdHandle LocalUserId;
+		FAccountId LocalUserId;
 		/* Implementation specific, used to specify the type of token to retrieve. See implementation specific documentation for valid values */
 		FString TokenType;
 	};
@@ -302,7 +302,7 @@ struct FCommerceOnPurchaseComplete
 public:
 	
 	/* Local user */
-	FOnlineAccountIdHandle LocalUserId;
+	FAccountId LocalUserId;
 	
 	/* If available, the transaction id for the completed purchase */
 	TOptional<FString> TransactionId;
diff --git a/Engine/Plugins/Online/Experimental/OnlineServices/Source/OnlineServicesInterface/Public/Online/ExternalUI.h b/Engine/Plugins/Online/Experimental/OnlineServices/Source/OnlineServicesInterface/Public/Online/ExternalUI.h
--- a/Engine/Plugins/Online/Experimental/OnlineServices/Source/OnlineServicesInterface/Public/Online/ExternalUI.h
+++ b/Engine/Plugins/Online/Experimental/OnlineServices/Source/OnlineServicesInterface/Public/Online/ExternalUI.h
@@ -37,7 +37,7 @@ struct FExternalUIShowFriendsUI
 	struct Params
 	{
 		/** Local user id to show the friends UI for */
-		FOnlineAccountIdHandle LocalUserId;
+		FAccountId LocalUserId;
 	};
 
 	struct Result
diff --git a/Engine/Plugins/Online/Experimental/OnlineServices/Source/OnlineServicesInterface/Public/Online/Friends.h b/Engine/Plugins/Online/Experimental/OnlineServices/Source/OnlineServicesInterface/Public/Online/Friends.h
--- a/Engine/Plugins/Online/Experimental/OnlineServices/Source/OnlineServicesInterface/Public/Online/Friends.h
+++ b/Engine/Plugins/Online/Experimental/OnlineServices/Source/OnlineServicesInterface/Public/Online/Friends.h
@@ -27,7 +27,7 @@ const TCHAR* LexToString(EFriendInviteStatus);
 struct FFriend
 {
 	/** Id of the friend */
-	FOnlineAccountIdHandle UserId;
+	FAccountId UserId;
 	/** Friendship invite status - if an invite is pending */
 	TOptional<EFriendInviteStatus> InviteStatus;
 };
@@ -40,7 +40,7 @@ struct FQueryFriends
 	struct Params
 	{
 		/** Account Id of the local user making the request */
-		FOnlineAccountIdHandle LocalUserId;
+		FAccountId LocalUserId;
 	};
 
 	/**
@@ -60,7 +60,7 @@ struct FGetFriends
 	struct Params
 	{
 		/** Account Id of the local user making the request */
-		FOnlineAccountIdHandle LocalUserId;
+		FAccountId LocalUserId;
 	};
 
 	/** Output struct for Friends::GetFriends */
@@ -79,9 +79,9 @@ struct FAddFriend
 	struct Params
 	{
 		/** Account Id of the local user making the request */
-		FOnlineAccountIdHandle LocalUserId;
+		FAccountId LocalUserId;
 		/** Friend to add */
-		FOnlineAccountIdHandle FriendId;
+		FAccountId FriendId;
 	};
 
 	/** Output struct for Friends::AddFriend */
@@ -96,13 +96,13 @@ struct FAddFriend
 struct FFriendsListUpdated
 {
 	/** Local user who's friends list changed */
-	FOnlineAccountIdHandle LocalUserId;
+	FAccountId LocalUserId;
 	/** Account ids of friends that have been added */
-	TArrayView<const FOnlineAccountIdHandle> AddedFriends;
+	TArrayView<const FAccountId> AddedFriends;
 	/** Account ids of friends that have been removed */
-	TArrayView<const FOnlineAccountIdHandle> RemovedFriends;
+	TArrayView<const FAccountId> RemovedFriends;
 	/** Account ids of friends that have been updated */
-	TArrayView<const FOnlineAccountIdHandle> UpdatedFriends;
+	TArrayView<const FAccountId> UpdatedFriends;
 };
 
 class IFriends
diff --git a/Engine/Plugins/Online/Experimental/OnlineServices/Source/OnlineServicesInterface/Public/Online/Leaderboards.h b/Engine/Plugins/Online/Experimental/OnlineServices/Source/OnlineServicesInterface/Public/Online/Leaderboards.h
--- a/Engine/Plugins/Online/Experimental/OnlineServices/Source/OnlineServicesInterface/Public/Online/Leaderboards.h
+++ b/Engine/Plugins/Online/Experimental/OnlineServices/Source/OnlineServicesInterface/Public/Online/Leaderboards.h
@@ -13,7 +13,7 @@ namespace UE::Online {
 struct FLeaderboardEntry
 {
 	/* The account id of this leaderboard entry */
-	FOnlineAccountIdHandle UserId;
+	FAccountId UserId;
 	/* The rank of this account */
 	int32 Rank;
 	/* The score of this account */
@@ -27,9 +27,9 @@ struct FReadEntriesForUsers
 	struct Params
 	{
 		/* Local user id */
-		FOnlineAccountIdHandle LocalUserId;
+		FAccountId LocalUserId;
 		/* The account ids of users */
-		TArray<FOnlineAccountIdHandle> UserIds;
+		TArray<FAccountId> UserIds;
 		/* The leaderboard name */
 		FString BoardName;
 	};
@@ -48,7 +48,7 @@ struct FReadEntriesAroundRank
 	struct Params
 	{
 		/* Local user id */
-		FOnlineAccountIdHandle LocalUserId;
+		FAccountId LocalUserId;
 		/* The rank index to start read */
 		uint32 Rank;
 		/* How many entries to fetch */
@@ -71,9 +71,9 @@ struct FReadEntriesAroundUser
 	struct Params
 	{
 		/* Local user id */
-		FOnlineAccountIdHandle LocalUserId;
+		FAccountId LocalUserId;
 		/* The account id of specified user */
-		FOnlineAccountIdHandle UserId;
+		FAccountId UserId;
 		/* The offset to the rank of user */
 		int32 Offset;
 		/* How many entries to fetch */
diff --git a/Engine/Plugins/Online/Experimental/OnlineServices/Source/OnlineServicesInterface/Public/Online/Lobbies.h b/Engine/Plugins/Online/Experimental/OnlineServices/Source/OnlineServicesInterface/Public/Online/Lobbies.h
--- a/Engine/Plugins/Online/Experimental/OnlineServices/Source/OnlineServicesInterface/Public/Online/Lobbies.h
+++ b/Engine/Plugins/Online/Experimental/OnlineServices/Source/OnlineServicesInterface/Public/Online/Lobbies.h
@@ -56,8 +56,8 @@ ONLINESERVICESINTERFACE_API void LexFromString(ELobbyMemberLeaveReason& OutLeave
 
 struct FLobbyMember
 {
-	FOnlineAccountIdHandle AccountId;
-	FOnlineAccountIdHandle PlatformAccountId;
+	FAccountId AccountId;
+	FAccountId PlatformAccountId;
 	FString PlatformDisplayName;
 	TMap<FLobbyAttributeId, FLobbyVariant> Attributes;
 };
@@ -65,19 +65,19 @@ struct FLobbyMember
 struct FLobby
 {
 	FOnlineLobbyIdHandle LobbyId;
-	FOnlineAccountIdHandle OwnerAccountId;
+	FAccountId OwnerAccountId;
 	FName LocalName;
 	FName SchemaName;
 	int32 MaxMembers;
 	ELobbyJoinPolicy JoinPolicy;
 	TMap<FLobbyAttributeId, FLobbyVariant> Attributes;
-	TMap<FOnlineAccountIdHandle, TSharedRef<const FLobbyMember>> Members;
+	TMap<FAccountId, TSharedRef<const FLobbyMember>> Members;
 };
 
 struct FJoinLobbyLocalUserData
 {
 	/** Local users who will be joining the lobby. */
-	FOnlineAccountIdHandle LocalUserId;
+	FAccountId LocalUserId;
 
 	/** Initial attributes. */
 	TMap<FLobbyAttributeId, FLobbyVariant> Attributes;
@@ -105,7 +105,7 @@ struct FCreateLobby
 	struct Params
 	{
 		/** The local user agent which will perform the action. */
-		FOnlineAccountIdHandle LocalUserId;
+		FAccountId LocalUserId;
 
 		/** The local name for the lobby. */
 		FName LocalName;
@@ -142,7 +142,7 @@ struct FFindLobbies
 	struct Params
 	{
 		/** The local user agent which will perform the action. */
-		FOnlineAccountIdHandle LocalUserId;
+		FAccountId LocalUserId;
 
 		/**
 		* Max results to return in one search.
@@ -154,7 +154,7 @@ struct FFindLobbies
 		TArray<FFindLobbySearchFilter> Filters;
 
 		/** Find lobbies containing the target user. */
-		TOptional<FOnlineAccountIdHandle> TargetUser;
+		TOptional<FAccountId> TargetUser;
 
 		/** Find join info for the target lobby id. */
 		TOptional<FOnlineLobbyIdHandle> LobbyId;
@@ -190,7 +190,7 @@ struct FJoinLobby
 	struct Params
 	{
 		/** The local user agent which will perform the action. */
-		FOnlineAccountIdHandle LocalUserId;
+		FAccountId LocalUserId;
 
 		/** The local name for the lobby. */
 		FName LocalName;
@@ -218,7 +218,7 @@ struct FLeaveLobby
 	struct Params
 	{
 		/** The local user agent which will perform the action. */
-		FOnlineAccountIdHandle LocalUserId;
+		FAccountId LocalUserId;
 
 		/** Id of the lobby to leave. */
 		FOnlineLobbyIdHandle LobbyId;
@@ -238,13 +238,13 @@ struct FInviteLobbyMember
 	struct Params
 	{
 		/** The local user agent which will perform the action. */
-		FOnlineAccountIdHandle LocalUserId;
+		FAccountId LocalUserId;
 
 		/** Id of the lobby for which the invitation will be sent. */
 		FOnlineLobbyIdHandle LobbyId;
 
 		/** Id of the player who will be sent the invitation. */
-		FOnlineAccountIdHandle TargetUserId;
+		FAccountId TargetUserId;
 	};
 
 	/** Output struct for Lobbies::InviteLobbyMember */
@@ -261,7 +261,7 @@ struct FDeclineLobbyInvitation
 	struct Params
 	{
 		/** The local user agent which will perform the action. */
-		FOnlineAccountIdHandle LocalUserId;
+		FAccountId LocalUserId;
 
 		/** Id of the lobby for which the invitations will be declined. */
 		FOnlineLobbyIdHandle LobbyId;
@@ -281,13 +281,13 @@ struct FKickLobbyMember
 	struct Params
 	{
 		/** The local user agent which will perform the action. */
-		FOnlineAccountIdHandle LocalUserId;
+		FAccountId LocalUserId;
 
 		/** Id of the lobby. */
 		FOnlineLobbyIdHandle LobbyId;
 
 		/** The target user to be kicked. */
-		FOnlineAccountIdHandle TargetUserId;
+		FAccountId TargetUserId;
 	};
 
 	/** Output struct for Lobbies::KickLobbyMember */
@@ -304,13 +304,13 @@ struct FPromoteLobbyMember
 	struct Params
 	{
 		/** The local user agent which will perform the action. */
-		FOnlineAccountIdHandle LocalUserId;
+		FAccountId LocalUserId;
 
 		/** Id of the lobby. */
 		FOnlineLobbyIdHandle LobbyId;
 
 		/** The target user to be promoted to owner. */
-		FOnlineAccountIdHandle TargetUserId;
+		FAccountId TargetUserId;
 	};
 
 	/** Output struct for Lobbies::PromoteLobbyMember */
@@ -342,7 +342,7 @@ struct FModifyLobbyJoinPolicy
 	struct Params
 	{
 		/** The local user agent which will perform the action. */
-		FOnlineAccountIdHandle LocalUserId;
+		FAccountId LocalUserId;
 
 		/** Id of the lobby. */
 		FOnlineLobbyIdHandle LobbyId;
@@ -365,7 +365,7 @@ struct FModifyLobbyAttributes
 	struct Params
 	{
 		/** The local user agent which will perform the action. */
-		FOnlineAccountIdHandle LocalUserId;
+		FAccountId LocalUserId;
 
 		/** Id of the lobby. */
 		FOnlineLobbyIdHandle LobbyId;
@@ -391,7 +391,7 @@ struct FModifyLobbyMemberAttributes
 	struct Params
 	{
 		/** The local user agent which will perform the action. */
-		FOnlineAccountIdHandle LocalUserId;
+		FAccountId LocalUserId;
 
 		/** Id of the lobby. */
 		FOnlineLobbyIdHandle LobbyId;
@@ -417,7 +417,7 @@ struct FGetJoinedLobbies
 	struct Params
 	{
 		/** The local user agent which will perform the action. */
-		FOnlineAccountIdHandle LocalUserId;
+		FAccountId LocalUserId;
 	};
 
 	/** Output struct for Lobbies::GetJoinedLobbies */
@@ -435,7 +435,7 @@ struct FGetReceivedInvitations
 	struct Params
 	{
 		/** The local user agent which will perform the action. */
-		FOnlineAccountIdHandle LocalUserId;
+		FAccountId LocalUserId;
 	};
 
 	/** Output struct for Lobbies::GetReceivedInvitations */
@@ -526,10 +526,10 @@ struct FLobbyMemberAttributesChanged
 struct FLobbyInvitationAdded
 {
 	/** The local user associated with the invitation. */
-	FOnlineAccountIdHandle LocalUserId;
+	FAccountId LocalUserId;
 
 	/** The user who sent the invitation. */
-	FOnlineAccountIdHandle SenderId;
+	FAccountId SenderId;
 
 	/** The invited lobby. */
 	TSharedRef<const FLobby> Lobby;
@@ -539,10 +539,10 @@ struct FLobbyInvitationAdded
 struct FLobbyInvitationRemoved
 {
 	/** The local user associated with the invitation. */
-	FOnlineAccountIdHandle LocalUserId;
+	FAccountId LocalUserId;
 
 	/** The user who sent the invitation. */
-	FOnlineAccountIdHandle SenderId;
+	FAccountId SenderId;
 
 	/** The invited lobby. */
 	TSharedRef<const FLobby> Lobby;
@@ -564,7 +564,7 @@ ONLINESERVICESINTERFACE_API void LexFromString(EUILobbyJoinRequestedSource& OutU
 struct FUILobbyJoinRequested
 {
 	/** The local user associated with the join request. */
-	FOnlineAccountIdHandle LocalUserId;
+	FAccountId LocalUserId;
 
 	/** The lobby the local user requested to join, or the online error if there was a failure */
 	TResult<TSharedRef<const FLobby>, FOnlineError> Result;
diff --git a/Engine/Plugins/Online/Experimental/OnlineServices/Source/OnlineServicesInterface/Public/Online/OnlineServices.h b/Engine/Plugins/Online/Experimental/OnlineServices/Source/OnlineServicesInterface/Public/Online/OnlineServices.h
--- a/Engine/Plugins/Online/Experimental/OnlineServices/Source/OnlineServicesInterface/Public/Online/OnlineServices.h
+++ b/Engine/Plugins/Online/Experimental/OnlineServices/Source/OnlineServicesInterface/Public/Online/OnlineServices.h
@@ -35,7 +35,7 @@ struct FGetResolvedConnectString
 
 	struct Params
 	{
-		FOnlineAccountIdHandle LocalUserId;
+		FAccountId LocalUserId;
 		FOnlineLobbyIdHandle LobbyId;
 		FOnlineSessionIdHandle SessionId;
 		FName PortType;
diff --git a/Engine/Plugins/Online/Experimental/OnlineServices/Source/OnlineServicesInterface/Public/Online/Presence.h b/Engine/Plugins/Online/Experimental/OnlineServices/Source/OnlineServicesInterface/Public/Online/Presence.h
--- a/Engine/Plugins/Online/Experimental/OnlineServices/Source/OnlineServicesInterface/Public/Online/Presence.h
+++ b/Engine/Plugins/Online/Experimental/OnlineServices/Source/OnlineServicesInterface/Public/Online/Presence.h
@@ -78,7 +78,7 @@ typedef TMap<FString, FPresenceVariant> FPresenceProperties;
 struct FUserPresence
 {
 	/** User whose presence this is */
-	FOnlineAccountIdHandle UserId;
+	FAccountId UserId;
 	/** Presence state */
 	EUserPresenceStatus Status = EUserPresenceStatus::Unknown;
 	/** Session state */
@@ -100,9 +100,9 @@ struct FQueryPresence
 	struct Params
 	{
 		/** Local user performing the query */
-		FOnlineAccountIdHandle LocalUserId;
+		FAccountId LocalUserId;
 		/** User to query the presence for */
-		FOnlineAccountIdHandle TargetUserId;
+		FAccountId TargetUserId;
 		/** If true, then future presence updates for this user will be propagated via the OnPresenceUpdated event */
 		bool bListenToChanges = true;
 	};
@@ -122,9 +122,9 @@ struct FBatchQueryPresence
 	struct Params
 	{
 		/** Local user performing the query */
-		FOnlineAccountIdHandle LocalUserId;
+		FAccountId LocalUserId;
 		/** Users to query the presence for */
-		TArray<FOnlineAccountIdHandle> TargetUserIds;
+		TArray<FAccountId> TargetUserIds;
 		/** If true, then future presence updates for these users will be propagated via the OnPresenceUpdated event */
 		bool bListenToChanges = true;
 	};
@@ -143,9 +143,9 @@ struct FGetCachedPresence
 	struct Params
 	{
 		/** Local user getting the presence */
-		FOnlineAccountIdHandle LocalUserId;
+		FAccountId LocalUserId;
 		/** User to get the presence for */
-		FOnlineAccountIdHandle TargetUserId;
+		FAccountId TargetUserId;
 	};
 
 	struct Result
@@ -162,7 +162,7 @@ struct FUpdatePresence
 	struct Params
 	{
 		/** Local user performing the query */
-		FOnlineAccountIdHandle LocalUserId;
+		FAccountId LocalUserId;
 		/** The new presence to send to the server */
 		TSharedRef<FUserPresence> Presence;
 	};
@@ -179,7 +179,7 @@ struct FPartialUpdatePresence
 	struct Params
 	{
 		/** Local user performing the query */
-		FOnlineAccountIdHandle LocalUserId;
+		FAccountId LocalUserId;
 
 
 		struct FMutations
@@ -293,7 +293,7 @@ inline TSharedRef<FUserPresence> ApplyPresenceMutations(const FUserPresence& Bas
 struct FPresenceUpdated
 {
 	/** Local user receiving the presence update */
-	FOnlineAccountIdHandle LocalUserId;
+	FAccountId LocalUserId;
 	/** Presence that has updated */
 	TSharedRef<const FUserPresence> UpdatedPresence;
 };
diff --git a/Engine/Plugins/Online/Experimental/OnlineServices/Source/OnlineServicesInterface/Public/Online/Privileges.h b/Engine/Plugins/Online/Experimental/OnlineServices/Source/OnlineServicesInterface/Public/Online/Privileges.h
--- a/Engine/Plugins/Online/Experimental/OnlineServices/Source/OnlineServicesInterface/Public/Online/Privileges.h
+++ b/Engine/Plugins/Online/Experimental/OnlineServices/Source/OnlineServicesInterface/Public/Online/Privileges.h
@@ -65,7 +65,7 @@ struct FQueryUserPrivilege
 	struct Params
 	{
 		/** User id whose privilege we want to query */
-		FOnlineAccountIdHandle LocalUserId;
+		FAccountId LocalUserId;
 		/** Privilege to query */
 		EUserPrivileges Privilege;
 	};
diff --git a/Engine/Plugins/Online/Experimental/OnlineServices/Source/OnlineServicesInterface/Public/Online/Sessions.h b/Engine/Plugins/Online/Experimental/OnlineServices/Source/OnlineServicesInterface/Public/Online/Sessions.h
--- a/Engine/Plugins/Online/Experimental/OnlineServices/Source/OnlineServicesInterface/Public/Online/Sessions.h
+++ b/Engine/Plugins/Online/Experimental/OnlineServices/Source/OnlineServicesInterface/Public/Online/Sessions.h
@@ -42,7 +42,7 @@ struct FSessionMember
 	FCustomSessionSettingsMap MemberSettings;
 };
 
-using FSessionMembersMap = TMap<FOnlineAccountIdHandle, FSessionMember>;
+using FSessionMembersMap = TMap<FAccountId, FSessionMember>;
 
 struct FSessionMemberUpdate
 {
@@ -52,7 +52,7 @@ struct FSessionMemberUpdate
 	FSessionMemberUpdate& operator+=(FSessionMemberUpdate&& UpdatedValue);
 };
 
-using FSessionMemberUpdatesMap = TMap<FOnlineAccountIdHandle, FSessionMemberUpdate>;
+using FSessionMemberUpdatesMap = TMap<FAccountId, FSessionMemberUpdate>;
 
 enum class ESessionState : uint8
 {
@@ -85,7 +85,7 @@ struct FRegisteredPlayer
 	bool bIsInSession = false;
 };
 
-using FRegisteredPlayersMap = TMap<FOnlineAccountIdHandle, FRegisteredPlayer>;
+using FRegisteredPlayersMap = TMap<FAccountId, FRegisteredPlayer>;
 
 struct FSessionSettingsUpdate
 {
@@ -124,12 +124,12 @@ struct FSessionSettingsUpdate
 	/** Updated values for session member info to change in the update operation*/
 	FSessionMemberUpdatesMap UpdatedSessionMembers;
 	/** Id handles for session members to be removed in the update operation*/
-	TArray<FOnlineAccountIdHandle> RemovedSessionMembers;
+	TArray<FAccountId> RemovedSessionMembers;
 
 	/** Updated values for registered players to change in the update operation*/
 	FRegisteredPlayersMap UpdatedRegisteredPlayers;
 	/** Id handles for registered players to be removed in the update operation*/
-	TArray<FOnlineAccountIdHandle> RemovedRegisteredPlayers;
+	TArray<FAccountId> RemovedRegisteredPlayers;
 
 	FSessionSettingsUpdate& operator+=(FSessionSettingsUpdate&& UpdatedValue);
 };
@@ -196,7 +196,7 @@ struct ONLINESERVICESINTERFACE_API FSession
 	FSession(const FSession& InSession);
 
 	/** The user who currently owns the session */
-	FOnlineAccountIdHandle OwnerUserId;
+	FAccountId OwnerUserId;
 
 	/** The id for the session, platform dependent */
 	FOnlineSessionIdHandle SessionId;
@@ -211,10 +211,10 @@ struct ONLINESERVICESINTERFACE_API FSession
 struct FSessionInvite
 {
 	/* The user which the invite got sent to */
-	FOnlineAccountIdHandle RecipientId;
+	FAccountId RecipientId;
 
 	/* The user which sent the invite */
-	FOnlineAccountIdHandle SenderId;
+	FAccountId SenderId;
 
 	/* The invite id handle, needed for retrieving session information and rejecting the invite */
 	FOnlineSessionInviteIdHandle InviteId;
@@ -263,7 +263,7 @@ struct FGetSessionById
 
 	struct Params
 	{
-		FOnlineAccountIdHandle LocalUserId;
+		FAccountId LocalUserId;
 
 		FOnlineSessionIdHandle IdHandle;
 	};
@@ -283,7 +283,7 @@ struct FCreateSession
 	struct Params
 	{
 		/** The local user agent which will perform the action. */
-		FOnlineAccountIdHandle LocalUserId;
+		FAccountId LocalUserId;
 
 		/** The local name for the session */
 		FName SessionName;
@@ -310,7 +310,7 @@ struct FUpdateSession
 	struct Params
 	{
 		/** The local user agent which will perform the action. */
-		FOnlineAccountIdHandle LocalUserId;
+		FAccountId LocalUserId;
 
 		/** The local name for the session */
 		FName SessionName;
@@ -334,13 +334,13 @@ struct FLeaveSession
 	struct Params
 	{
 		/* The local user agent which leaves the session*/
-		FOnlineAccountIdHandle LocalUserId;
+		FAccountId LocalUserId;
 
 		/* The local name for the session. */
 		FName SessionName;
 
 		/** Ids for all local users who will leave the session (includes the main caller) */
-		TArray<FOnlineAccountIdHandle> LocalUsers;
+		TArray<FAccountId> LocalUsers;
 
 		/* Whether the call should attempt to destroy the session instead of just leave it */
 		bool bDestroySession;
@@ -359,7 +359,7 @@ struct FFindSessions
 	struct Params
 	{
 		/* The local user agent which starts the session search*/
-		FOnlineAccountIdHandle LocalUserId;
+		FAccountId LocalUserId;
 
 		/* Maximum number of results to return in one search */
 		uint32 MaxResults;
@@ -371,7 +371,7 @@ struct FFindSessions
 		TArray<FFindSessionsSearchFilter> Filters;
 
 		/* Find sessions containing the target user. */
-		TOptional<FOnlineAccountIdHandle> TargetUser;
+		TOptional<FAccountId> TargetUser;
 
 		/* Find join info for the target session id. */
 		TOptional<FOnlineSessionIdHandle> SessionId;
@@ -390,7 +390,7 @@ struct FStartMatchmaking
 	struct Params
 	{
 		/** The local user agent which will perform the action. */
-		FOnlineAccountIdHandle LocalUserId;
+		FAccountId LocalUserId;
 
 		/* Information for all local users who will join the session (includes the session creator) */
 		FSessionMembersMap LocalUsers;
@@ -420,7 +420,7 @@ struct FJoinSession
 	struct Params
 	{
 		/* The local user agent which starts the join operation*/
-		FOnlineAccountIdHandle LocalUserId;
+		FAccountId LocalUserId;
 
 		/* Local name for the session */
 		FName SessionName;
@@ -447,7 +447,7 @@ struct FAddSessionMembers
 	struct Params
 	{
 		/* The local user agent */
-		FOnlineAccountIdHandle LocalUserId;
+		FAccountId LocalUserId;
 
 		/* Local name for the session */
 		FName SessionName;
@@ -472,13 +472,13 @@ struct FRemoveSessionMembers
 	struct Params
 	{
 		/* The local user agent */
-		FOnlineAccountIdHandle LocalUserId;
+		FAccountId LocalUserId;
 
 		/* Local name for the session */
 		FName SessionName;
 
 		/* Id handles for the session members to be removed from the session */
-		TArray<FOnlineAccountIdHandle> SessionMemberIds;
+		TArray<FAccountId> SessionMemberIds;
 
 		/** Whether or not the session members should also be removed from the list of registered players. True by default*/
 		bool bUnregisterPlayers = true;
@@ -497,13 +497,13 @@ struct FSendSessionInvite
 	struct Params
 	{
 		/* The local user agent which sends the invite*/
-		FOnlineAccountIdHandle LocalUserId;
+		FAccountId LocalUserId;
 
 		/* The local name for the session. */
 		FName SessionName;
 
 		/* Array of id handles for users to which the invites will be sent */
-		TArray<FOnlineAccountIdHandle> TargetUsers;
+		TArray<FAccountId> TargetUsers;
 	};
 
 	struct Result
@@ -519,7 +519,7 @@ struct FGetSessionInvites
 	struct Params
 	{
 		/* The local user agent */
-		FOnlineAccountIdHandle LocalUserId;
+		FAccountId LocalUserId;
 	};
 
 	struct Result
@@ -536,7 +536,7 @@ struct FRejectSessionInvite
 	struct Params
 	{
 		/* The local user agent which started the query*/
-		FOnlineAccountIdHandle LocalUserId;
+		FAccountId LocalUserId;
 
 		/* The id handle for the invite to be rejected */
 		FOnlineSessionInviteIdHandle SessionInviteId;
@@ -557,7 +557,7 @@ struct FRegisterPlayers
 		FName SessionName;
 
 		/* Array of users which will be registered */
-		TArray<FOnlineAccountIdHandle> TargetUsers;
+		TArray<FAccountId> TargetUsers;
 
 		/* Whether a slot should be saved for the registered players */
 		bool bReserveSlot;
@@ -578,7 +578,7 @@ struct FUnregisterPlayers
 		FName SessionName;
 
 		/* Array of users which will be unregistered */
-		TArray<FOnlineAccountIdHandle> TargetUsers;
+		TArray<FAccountId> TargetUsers;
 
 		/* Whether unregistered players should be removed from the session, if they are in it */
 		bool bRemoveUnregisteredPlayers;
@@ -594,7 +594,7 @@ struct FUnregisterPlayers
 struct FSessionJoined
 {
 	/* The local users which joined the session */
-	TArray<FOnlineAccountIdHandle> LocalUserIds;
+	TArray<FAccountId> LocalUserIds;
 
 	/* A shared reference to the session joined. */
 	TSharedRef<const FSession> Session;
@@ -605,7 +605,7 @@ struct FSessionJoined
 struct FSessionLeft
 {
 	/* The local users which left the session */
-	TArray<FOnlineAccountIdHandle> LocalUserIds;
+	TArray<FAccountId> LocalUserIds;
 };
 
 struct FSessionUpdated
@@ -622,7 +622,7 @@ struct FSessionUpdated
 struct FSessionInviteReceived
 {
 	/* The local user which received the invite */
-	FOnlineAccountIdHandle LocalUserId;
+	FAccountId LocalUserId;
 
 	/** The session invite the local user was sent, or the online error if there was a failure retrieving the session for it*/
 	TSharedRef<const FSessionInvite> SessionInvite;
@@ -644,7 +644,7 @@ ONLINESERVICESINTERFACE_API void LexFromString(EUISessionJoinRequestedSource& Ou
 struct FUISessionJoinRequested
 {
 	/** The local user associated with the join request. */
-	FOnlineAccountIdHandle LocalUserId;
+	FAccountId LocalUserId;
 
 	/** The session the local user requested to join, or the online error if there was a failure retrieving it */
 	TResult<TSharedRef<const FSession>, FOnlineError> Result;
diff --git a/Engine/Plugins/Online/Experimental/OnlineServices/Source/OnlineServicesInterface/Public/Online/Stats.h b/Engine/Plugins/Online/Experimental/OnlineServices/Source/OnlineServicesInterface/Public/Online/Stats.h
--- a/Engine/Plugins/Online/Experimental/OnlineServices/Source/OnlineServicesInterface/Public/Online/Stats.h
+++ b/Engine/Plugins/Online/Experimental/OnlineServices/Source/OnlineServicesInterface/Public/Online/Stats.h
@@ -14,7 +14,7 @@ using FStatValue = FSchemaVariant;
 struct FUserStats
 {
 	/* The user id */
-	FOnlineAccountIdHandle UserId;
+	FAccountId UserId;
 	/* The stats of one account */
 	TMap<FString, FStatValue> Stats;
 };
@@ -26,7 +26,7 @@ struct FUpdateStats
 	struct Params
 	{
 		/* Local user id */
-		FOnlineAccountIdHandle LocalUserId;
+		FAccountId LocalUserId;
 		/* The stats of users to update */
 		TArray<FUserStats> UpdateUsersStats;
 	};
@@ -43,9 +43,9 @@ struct FQueryStats
 	struct Params
 	{
 		/* Local user id */
-		FOnlineAccountIdHandle LocalUserId;
+		FAccountId LocalUserId;
 		/* Target user id to query */
-		FOnlineAccountIdHandle TargetUserId;
+		FAccountId TargetUserId;
 		/* The stats to query */
 		TArray<FString> StatNames;
 	};
@@ -64,9 +64,9 @@ struct FBatchQueryStats
 	struct Params
 	{
 		/* Local user id */
-		FOnlineAccountIdHandle LocalUserId;
+		FAccountId LocalUserId;
 		/* The user ids to query */
-		TArray<FOnlineAccountIdHandle> TargetUserIds;
+		TArray<FAccountId> TargetUserIds;
 		/* The stats to query */
 		TArray<FString> StatNames;
 	};
@@ -86,7 +86,7 @@ struct FResetStats
 	struct Params
 	{
 		/* Local user id */
-		FOnlineAccountIdHandle LocalUserId;
+		FAccountId LocalUserId;
 	};
 
 	struct Result
diff --git a/Engine/Plugins/Online/Experimental/OnlineServices/Source/OnlineServicesInterface/Public/Online/TitleFile.h b/Engine/Plugins/Online/Experimental/OnlineServices/Source/OnlineServicesInterface/Public/Online/TitleFile.h
--- a/Engine/Plugins/Online/Experimental/OnlineServices/Source/OnlineServicesInterface/Public/Online/TitleFile.h
+++ b/Engine/Plugins/Online/Experimental/OnlineServices/Source/OnlineServicesInterface/Public/Online/TitleFile.h
@@ -15,7 +15,7 @@ struct FTitleFileEnumerateFiles
 	struct Params
 	{
 		/** Local user performing the operation */
-		FOnlineAccountIdHandle LocalUserId;
+		FAccountId LocalUserId;
 	};
 
 	struct Result
@@ -30,7 +30,7 @@ struct FTitleFileGetEnumeratedFiles
 	struct Params
 	{
 		/** Local user performing the operation */
-		FOnlineAccountIdHandle LocalUserId;
+		FAccountId LocalUserId;
 	};
 
 	struct Result
@@ -52,7 +52,7 @@ struct FTitleFileReadFile
 	struct Params
 	{
 		/** Local user performing the operation */
-		FOnlineAccountIdHandle LocalUserId;
+		FAccountId LocalUserId;
 		/** File to be opened */
 		FString Filename;
 	};
diff --git a/Engine/Plugins/Online/Experimental/OnlineServices/Source/OnlineServicesInterface/Public/Online/UserFile.h b/Engine/Plugins/Online/Experimental/OnlineServices/Source/OnlineServicesInterface/Public/Online/UserFile.h
--- a/Engine/Plugins/Online/Experimental/OnlineServices/Source/OnlineServicesInterface/Public/Online/UserFile.h
+++ b/Engine/Plugins/Online/Experimental/OnlineServices/Source/OnlineServicesInterface/Public/Online/UserFile.h
@@ -15,7 +15,7 @@ struct FUserFileEnumerateFiles
 	struct Params
 	{
 		/** Local user performing the operation */
-		FOnlineAccountIdHandle LocalUserId;
+		FAccountId LocalUserId;
 	};
 
 	struct Result
@@ -30,7 +30,7 @@ struct FUserFileGetEnumeratedFiles
 	struct Params
 	{
 		/** Local user performing the operation */
-		FOnlineAccountIdHandle LocalUserId;
+		FAccountId LocalUserId;
 	};
 
 	struct Result
@@ -52,7 +52,7 @@ struct FUserFileReadFile
 	struct Params
 	{
 		/** Local user performing the operation */
-		FOnlineAccountIdHandle LocalUserId;
+		FAccountId LocalUserId;
 		/** File to be read */
 		FString Filename;
 	};
@@ -71,7 +71,7 @@ struct FUserFileWriteFile
 	struct Params
 	{
 		/** Local user performing the operation */
-		FOnlineAccountIdHandle LocalUserId;
+		FAccountId LocalUserId;
 		/** File to be written */
 		FString Filename;
 		/** Data to write */
@@ -90,7 +90,7 @@ struct FUserFileCopyFile
 	struct Params
 	{
 		/** Local user performing the operation */
-		FOnlineAccountIdHandle LocalUserId;
+		FAccountId LocalUserId;
 		/** File to copy contents from */
 		FString SourceFilename;
 		/** File to copy contents to */
@@ -109,7 +109,7 @@ struct FUserFileDeleteFile
 	struct Params
 	{
 		/** Local user performing the operation */
-		FOnlineAccountIdHandle LocalUserId;
+		FAccountId LocalUserId;
 		/** File to delete */
 		FString Filename;
 	};
diff --git a/Engine/Plugins/Online/Experimental/OnlineServices/Source/OnlineServicesInterface/Public/Online/UserInfo.h b/Engine/Plugins/Online/Experimental/OnlineServices/Source/OnlineServicesInterface/Public/Online/UserInfo.h
--- a/Engine/Plugins/Online/Experimental/OnlineServices/Source/OnlineServicesInterface/Public/Online/UserInfo.h
+++ b/Engine/Plugins/Online/Experimental/OnlineServices/Source/OnlineServicesInterface/Public/Online/UserInfo.h
@@ -11,7 +11,7 @@ namespace UE::Online {
 struct FUserInfo
 {
 	/* Account Id */
-	FOnlineAccountIdHandle UserId;
+	FAccountId UserId;
 	/* Display Name of the User */
 	FString DisplayName;
 };
@@ -24,9 +24,9 @@ struct FQueryUserInfo
 	struct Params
 	{
 		/* Local user to query users for */
-		FOnlineAccountIdHandle LocalUserId;
+		FAccountId LocalUserId;
 		/* List of User ids to Query*/
-		TArray<FOnlineAccountIdHandle> UserIds;
+		TArray<FAccountId> UserIds;
 	};
 
 	/**
@@ -46,9 +46,9 @@ struct FGetUserInfo
 	struct Params
 	{
 		/* Local user */
-		FOnlineAccountIdHandle LocalUserId;
+		FAccountId LocalUserId;
 		/* User id to get info for*/
-		FOnlineAccountIdHandle UserId;
+		FAccountId UserId;
 	};
 
 	/**
@@ -69,9 +69,9 @@ struct FQueryUserAvatar
 	struct Params
 	{
 		/* Local user to query users for */
-		FOnlineAccountIdHandle LocalUserId;
+		FAccountId LocalUserId;
 		/* List of User ids to Query*/
-		TArray<FOnlineAccountIdHandle> UserIds;
+		TArray<FAccountId> UserIds;
 	};
 
 	/**
@@ -91,9 +91,9 @@ struct FGetUserAvatar
 	struct Params
 	{
 		/* Local user */
-		FOnlineAccountIdHandle LocalUserId;
+		FAccountId LocalUserId;
 		/* User id to get avatar for*/
-		FOnlineAccountIdHandle UserId;
+		FAccountId UserId;
 	};
 
 	/**
@@ -114,9 +114,9 @@ struct FShowUserProfile
 	struct Params
 	{
 		/* Local user */
-		FOnlineAccountIdHandle LocalUserId;
+		FAccountId LocalUserId;
 		/* User id to display the profile for*/
-		FOnlineAccountIdHandle UserId;
+		FAccountId UserId;
 	};
 
 	/**
diff --git a/Engine/Plugins/Online/Experimental/OnlineServicesEOS/Source/Private/Online/AuthEOS.cpp b/Engine/Plugins/Online/Experimental/OnlineServicesEOS/Source/Private/Online/AuthEOS.cpp
--- a/Engine/Plugins/Online/Experimental/OnlineServicesEOS/Source/Private/Online/AuthEOS.cpp
+++ b/Engine/Plugins/Online/Experimental/OnlineServicesEOS/Source/Private/Online/AuthEOS.cpp
@@ -257,14 +257,14 @@ TOnlineAsyncOpHandle<FAuthLogin> FAuthEOS::Login(FAuthLogin::Params&& Params)
 }
 
 template <typename IdType>
-TFuture<FOnlineAccountIdHandle> ResolveAccountIdImpl(FAuthEOS& AuthEOS, const FOnlineAccountIdHandle& LocalUserId, const IdType InId)
+TFuture<FAccountId> ResolveAccountIdImpl(FAuthEOS& AuthEOS, const FAccountId& LocalUserId, const IdType InId)
 {
-	TPromise<FOnlineAccountIdHandle> Promise;
-	TFuture<FOnlineAccountIdHandle> Future = Promise.GetFuture();
+	TPromise<FAccountId> Promise;
+	TFuture<FAccountId> Future = Promise.GetFuture();
 
-	AuthEOS.ResolveAccountIds(LocalUserId, { InId }).Next([Promise = MoveTemp(Promise)](const TArray<FOnlineAccountIdHandle>& AccountIds) mutable
+	AuthEOS.ResolveAccountIds(LocalUserId, { InId }).Next([Promise = MoveTemp(Promise)](const TArray<FAccountId>& AccountIds) mutable
 	{
-		FOnlineAccountIdHandle Result;
+		FAccountId Result;
 		if (AccountIds.Num() == 1)
 		{
 			Result = AccountIds[0];
@@ -347,22 +347,22 @@ TOnlineAsyncOpHandle<FAuthQueryExternalAuthToken> FAuthEOS::QueryExternalAuthTok
 	return Op->GetHandle();
 }
 
-TFuture<FOnlineAccountIdHandle> FAuthEOS::ResolveAccountId(const FOnlineAccountIdHandle& LocalUserId, const EOS_EpicAccountId EpicAccountId)
+TFuture<FAccountId> FAuthEOS::ResolveAccountId(const FAccountId& LocalUserId, const EOS_EpicAccountId EpicAccountId)
 {
 	return ResolveAccountIdImpl(*this, LocalUserId, EpicAccountId);
 }
 
-TFuture<FOnlineAccountIdHandle> FAuthEOS::ResolveAccountId(const FOnlineAccountIdHandle& LocalUserId, const EOS_ProductUserId ProductUserId)
+TFuture<FAccountId> FAuthEOS::ResolveAccountId(const FAccountId& LocalUserId, const EOS_ProductUserId ProductUserId)
 {
 	return ResolveAccountIdImpl(*this, LocalUserId, ProductUserId);
 }
 
 using FEpicAccountIdStrBuffer = char[EOS_EPICACCOUNTID_MAX_LENGTH + 1];
 
-TFuture<TArray<FOnlineAccountIdHandle>> FAuthEOS::ResolveAccountIds(const FOnlineAccountIdHandle& LocalUserId, const TArray<EOS_EpicAccountId>& InEpicAccountIds)
+TFuture<TArray<FAccountId>> FAuthEOS::ResolveAccountIds(const FAccountId& LocalUserId, const TArray<EOS_EpicAccountId>& InEpicAccountIds)
 {
 	// Search for all the account id's
-	TArray<FOnlineAccountIdHandle> AccountIdHandles;
+	TArray<FAccountId> AccountIdHandles;
 	AccountIdHandles.Reserve(InEpicAccountIds.Num());
 	TArray<EOS_EpicAccountId> MissingEpicAccountIds;
 	MissingEpicAccountIds.Reserve(InEpicAccountIds.Num());
@@ -370,10 +370,10 @@ TFuture<TArray<FOnlineAccountIdHandle>> FAuthEOS::ResolveAccountIds(const FOnlin
 	{
 		if (!EOS_EpicAccountId_IsValid(EpicAccountId))
 		{
-			return MakeFulfilledPromise<TArray<FOnlineAccountIdHandle>>().GetFuture();
+			return MakeFulfilledPromise<TArray<FAccountId>>().GetFuture();
 		}
 
-		FOnlineAccountIdHandle Found = FindAccountId(EpicAccountId);
+		FAccountId Found = FindAccountId(EpicAccountId);
 		if (!Found.IsValid())
 		{
 			MissingEpicAccountIds.Emplace(EpicAccountId);
@@ -383,18 +383,18 @@ TFuture<TArray<FOnlineAccountIdHandle>> FAuthEOS::ResolveAccountIds(const FOnlin
 	if (MissingEpicAccountIds.IsEmpty())
 	{
 		// We have them all, so we can just return
-		return MakeFulfilledPromise<TArray<FOnlineAccountIdHandle>>(MoveTemp(AccountIdHandles)).GetFuture();
+		return MakeFulfilledPromise<TArray<FAccountId>>(MoveTemp(AccountIdHandles)).GetFuture();
 	}
 
 	// If we failed to find all the handles, we need to query, which requires a valid LocalUserId
 	if (!ValidateOnlineId(LocalUserId))
 	{
 		checkNoEntry();
-		return MakeFulfilledPromise<TArray<FOnlineAccountIdHandle>>().GetFuture();
+		return MakeFulfilledPromise<TArray<FAccountId>>().GetFuture();
 	}
 
-	TPromise<TArray<FOnlineAccountIdHandle>> Promise;
-	TFuture<TArray<FOnlineAccountIdHandle>> Future = Promise.GetFuture();
+	TPromise<TArray<FAccountId>> Promise;
+	TFuture<TArray<FAccountId>> Future = Promise.GetFuture();
 
 	TArray<FEpicAccountIdStrBuffer> EpicAccountIdStrsToQuery;
 	EpicAccountIdStrsToQuery.Reserve(MissingEpicAccountIds.Num());
@@ -406,7 +406,7 @@ TFuture<TArray<FOnlineAccountIdHandle>> FAuthEOS::ResolveAccountIds(const FOnlin
 			EOS_EpicAccountId_ToString(EpicAccountId, EpicAccountIdStr, &BufferSize) != EOS_EResult::EOS_Success)
 		{
 			checkNoEntry();
-			return MakeFulfilledPromise<TArray<FOnlineAccountIdHandle>>().GetFuture();
+			return MakeFulfilledPromise<TArray<FAccountId>>().GetFuture();
 		}
 	}
 
@@ -424,7 +424,7 @@ TFuture<TArray<FOnlineAccountIdHandle>> FAuthEOS::ResolveAccountIds(const FOnlin
 	EOS_Async(EOS_Connect_QueryExternalAccountMappings, ConnectHandle, Options,
 	[this, WeakThis = AsWeak(), InEpicAccountIds, Promise = MoveTemp(Promise)](const EOS_Connect_QueryExternalAccountMappingsCallbackInfo* Data) mutable -> void
 	{
-		TArray<FOnlineAccountIdHandle> AccountIds;
+		TArray<FAccountId> AccountIds;
 		if (const TSharedPtr<IAuth> StrongThis = WeakThis.Pin())
 		{
 			AccountIds.Reserve(InEpicAccountIds.Num());
@@ -438,7 +438,7 @@ TFuture<TArray<FOnlineAccountIdHandle>> FAuthEOS::ResolveAccountIds(const FOnlin
 
 				for (const EOS_EpicAccountId EpicAccountId : InEpicAccountIds)
 				{
-					FOnlineAccountIdHandle AccountId = FindAccountId(EpicAccountId);
+					FAccountId AccountId = FindAccountId(EpicAccountId);
 					if (!AccountId.IsValid())
 					{
 						FEpicAccountIdStrBuffer EpicAccountIdStr;
@@ -462,10 +462,10 @@ TFuture<TArray<FOnlineAccountIdHandle>> FAuthEOS::ResolveAccountIds(const FOnlin
 	return Future;
 }
 
-TFuture<TArray<FOnlineAccountIdHandle>> FAuthEOS::ResolveAccountIds(const FOnlineAccountIdHandle& LocalUserId, const TArray<EOS_ProductUserId>& InProductUserIds)
+TFuture<TArray<FAccountId>> FAuthEOS::ResolveAccountIds(const FAccountId& LocalUserId, const TArray<EOS_ProductUserId>& InProductUserIds)
 {
 	// Search for all the account id's
-	TArray<FOnlineAccountIdHandle> AccountIdHandles;
+	TArray<FAccountId> AccountIdHandles;
 	AccountIdHandles.Reserve(InProductUserIds.Num());
 	TArray<EOS_ProductUserId> MissingProductUserIds;
 	MissingProductUserIds.Reserve(InProductUserIds.Num());
@@ -473,10 +473,10 @@ TFuture<TArray<FOnlineAccountIdHandle>> FAuthEOS::ResolveAccountIds(const FOnlin
 	{
 		if (!EOS_ProductUserId_IsValid(ProductUserId))
 		{
-			return MakeFulfilledPromise<TArray<FOnlineAccountIdHandle>>().GetFuture();
+			return MakeFulfilledPromise<TArray<FAccountId>>().GetFuture();
 		}
 
-		FOnlineAccountIdHandle Found = FindAccountId(ProductUserId);
+		FAccountId Found = FindAccountId(ProductUserId);
 		if (!Found.IsValid())
 		{
 			MissingProductUserIds.Emplace(ProductUserId);
@@ -486,18 +486,18 @@ TFuture<TArray<FOnlineAccountIdHandle>> FAuthEOS::ResolveAccountIds(const FOnlin
 	if (MissingProductUserIds.IsEmpty())
 	{
 		// We have them all, so we can just return
-		return MakeFulfilledPromise<TArray<FOnlineAccountIdHandle>>(MoveTemp(AccountIdHandles)).GetFuture();
+		return MakeFulfilledPromise<TArray<FAccountId>>(MoveTemp(AccountIdHandles)).GetFuture();
 	}
 
 	// If we failed to find all the handles, we need to query, which requires a valid LocalUserId
 	if (!ValidateOnlineId(LocalUserId))
 	{
 		checkNoEntry();
-		return MakeFulfilledPromise<TArray<FOnlineAccountIdHandle>>().GetFuture();
+		return MakeFulfilledPromise<TArray<FAccountId>>().GetFuture();
 	}
 
-	TPromise<TArray<FOnlineAccountIdHandle>> Promise;
-	TFuture<TArray<FOnlineAccountIdHandle>> Future = Promise.GetFuture();
+	TPromise<TArray<FAccountId>> Promise;
+	TFuture<TArray<FAccountId>> Future = Promise.GetFuture();
 
 	EOS_Connect_QueryProductUserIdMappingsOptions Options = {};
 	Options.ApiVersion = EOS_CONNECT_QUERYPRODUCTUSERIDMAPPINGS_API_LATEST;
@@ -509,7 +509,7 @@ TFuture<TArray<FOnlineAccountIdHandle>> FAuthEOS::ResolveAccountIds(const FOnlin
 	EOS_Async(EOS_Connect_QueryProductUserIdMappings, ConnectHandle, Options,
 	[this, WeakThis = AsWeak(), InProductUserIds, Promise = MoveTemp(Promise)](const EOS_Connect_QueryProductUserIdMappingsCallbackInfo* Data) mutable -> void
 	{
-		TArray<FOnlineAccountIdHandle> AccountIds;
+		TArray<FAccountId> AccountIds;
 		if (const TSharedPtr<IAuth> StrongThis = WeakThis.Pin())
 		{
 			if (Data->ResultCode == EOS_EResult::EOS_Success)
@@ -522,7 +522,7 @@ TFuture<TArray<FOnlineAccountIdHandle>> FAuthEOS::ResolveAccountIds(const FOnlin
 
 				for (const EOS_ProductUserId ProductUserId : InProductUserIds)
 				{
-					FOnlineAccountIdHandle AccountId = FindAccountId(ProductUserId);
+					FAccountId AccountId = FindAccountId(ProductUserId);
 					if (!AccountId.IsValid())
 					{
 						Options.TargetProductUserId = ProductUserId;
@@ -552,50 +552,50 @@ TFuture<TArray<FOnlineAccountIdHandle>> FAuthEOS::ResolveAccountIds(const FOnlin
 }
 
 template<typename ParamType>
-TFunction<TFuture<FOnlineAccountIdHandle>(FOnlineAsyncOp& InAsyncOp, const ParamType& Param)> ResolveIdFnImpl(FAuthEOS* AuthEOS)
+TFunction<TFuture<FAccountId>(FOnlineAsyncOp& InAsyncOp, const ParamType& Param)> ResolveIdFnImpl(FAuthEOS* AuthEOS)
 {
 	return [AuthEOS](FOnlineAsyncOp& InAsyncOp, const ParamType& Param)
 	{
-		const FOnlineAccountIdHandle* LocalUserIdPtr = InAsyncOp.Data.Get<FOnlineAccountIdHandle>(TEXT("LocalUserId"));
+		const FAccountId* LocalUserIdPtr = InAsyncOp.Data.Get<FAccountId>(TEXT("LocalUserId"));
 		if (!ensure(LocalUserIdPtr))
 		{
-			return MakeFulfilledPromise<FOnlineAccountIdHandle>().GetFuture();
+			return MakeFulfilledPromise<FAccountId>().GetFuture();
 		}
 		return AuthEOS->ResolveAccountId(*LocalUserIdPtr, Param);
 	};
 }
-TFunction<TFuture<FOnlineAccountIdHandle>(FOnlineAsyncOp& InAsyncOp, const EOS_EpicAccountId& ProductUserId)> FAuthEOS::ResolveEpicIdFn()
+TFunction<TFuture<FAccountId>(FOnlineAsyncOp& InAsyncOp, const EOS_EpicAccountId& ProductUserId)> FAuthEOS::ResolveEpicIdFn()
 {
 	return ResolveIdFnImpl<EOS_EpicAccountId>(this);
 }
-TFunction<TFuture<FOnlineAccountIdHandle>(FOnlineAsyncOp& InAsyncOp, const EOS_ProductUserId& ProductUserId)> FAuthEOS::ResolveProductIdFn()
+TFunction<TFuture<FAccountId>(FOnlineAsyncOp& InAsyncOp, const EOS_ProductUserId& ProductUserId)> FAuthEOS::ResolveProductIdFn()
 {
 	return ResolveIdFnImpl<EOS_ProductUserId>(this);
 }
 
 template<typename ParamType>
-TFunction<TFuture<TArray<FOnlineAccountIdHandle>>(FOnlineAsyncOp& InAsyncOp, const TArray<ParamType>& Param)> ResolveIdsFnImpl(FAuthEOS* AuthEOS)
+TFunction<TFuture<TArray<FAccountId>>(FOnlineAsyncOp& InAsyncOp, const TArray<ParamType>& Param)> ResolveIdsFnImpl(FAuthEOS* AuthEOS)
 {
 	return [AuthEOS](FOnlineAsyncOp& InAsyncOp, const TArray<ParamType>& Param)
 	{
-		const FOnlineAccountIdHandle* LocalUserIdPtr = InAsyncOp.Data.Get<FOnlineAccountIdHandle>(TEXT("LocalUserId"));
+		const FAccountId* LocalUserIdPtr = InAsyncOp.Data.Get<FAccountId>(TEXT("LocalUserId"));
 		if (!ensure(LocalUserIdPtr))
 		{
-			return MakeFulfilledPromise<TArray<FOnlineAccountIdHandle>>().GetFuture();
+			return MakeFulfilledPromise<TArray<FAccountId>>().GetFuture();
 		}
 		return AuthEOS->ResolveAccountIds(*LocalUserIdPtr, Param);
 	};
 }
-TFunction<TFuture<TArray<FOnlineAccountIdHandle>>(FOnlineAsyncOp& InAsyncOp, const TArray<EOS_EpicAccountId>& EpicAccountIds)> FAuthEOS::ResolveEpicIdsFn()
+TFunction<TFuture<TArray<FAccountId>>(FOnlineAsyncOp& InAsyncOp, const TArray<EOS_EpicAccountId>& EpicAccountIds)> FAuthEOS::ResolveEpicIdsFn()
 {
 	return ResolveIdsFnImpl<EOS_EpicAccountId>(this);
 }
-TFunction<TFuture<TArray<FOnlineAccountIdHandle>>(FOnlineAsyncOp& InAsyncOp, const TArray<EOS_ProductUserId>& ProductUserIds)> FAuthEOS::ResolveProductIdsFn()
+TFunction<TFuture<TArray<FAccountId>>(FOnlineAsyncOp& InAsyncOp, const TArray<EOS_ProductUserId>& ProductUserIds)> FAuthEOS::ResolveProductIdsFn()
 {
 	return ResolveIdsFnImpl<EOS_ProductUserId>(this);
 }
 
-FOnlineAccountIdHandle FAuthEOS::CreateAccountId(const EOS_EpicAccountId EpicAccountId, const EOS_ProductUserId ProductUserId)
+FAccountId FAuthEOS::CreateAccountId(const EOS_EpicAccountId EpicAccountId, const EOS_ProductUserId ProductUserId)
 {
 	return FOnlineAccountIdRegistryEOS::Get().FindOrAddAccountId(EpicAccountId, ProductUserId);
 }
diff --git a/Engine/Plugins/Online/Experimental/OnlineServicesEOS/Source/Private/Online/FriendsEOS.cpp b/Engine/Plugins/Online/Experimental/OnlineServicesEOS/Source/Private/Online/FriendsEOS.cpp
--- a/Engine/Plugins/Online/Experimental/OnlineServicesEOS/Source/Private/Online/FriendsEOS.cpp
+++ b/Engine/Plugins/Online/Experimental/OnlineServicesEOS/Source/Private/Online/FriendsEOS.cpp
@@ -45,9 +45,9 @@ void FFriendsEOS::Initialize()
 	{
 		FFriendsEOS* This = reinterpret_cast<FFriendsEOS*>(Data->ClientData);
 
-		const FOnlineAccountIdHandle LocalUserId = FindAccountIdChecked(Data->LocalUserId);
+		const FAccountId LocalUserId = FindAccountIdChecked(Data->LocalUserId);
 		This->Services.Get<FAuthEOS>()->ResolveAccountId(LocalUserId, Data->TargetUserId)
-		.Next([This, LocalUserId, PreviousStatus = Data->PreviousStatus, CurrentStatus = Data->CurrentStatus](const FOnlineAccountIdHandle& FriendUserId)
+		.Next([This, LocalUserId, PreviousStatus = Data->PreviousStatus, CurrentStatus = Data->CurrentStatus](const FAccountId& FriendUserId)
 		{
 			UE_LOG(LogTemp, Warning, TEXT("OnEOSFriendsUpdate: LocalUserId=[%s] FriendUserId=[%s]"), *ToLogString(LocalUserId), *ToLogString(FriendUserId));
 			This->OnEOSFriendsUpdate(LocalUserId, FriendUserId, PreviousStatus, CurrentStatus);
@@ -115,18 +115,18 @@ TOnlineAsyncOpHandle<FQueryFriends> FFriendsEOS::QueryFriends(FQueryFriends::Par
 			}
 			return MakeFulfilledPromise<TArray<EOS_EpicAccountId>>(CurrentFriends).GetFuture();
 		}).Then(Services.Get<FAuthEOS>()->ResolveEpicIdsFn())
-		.Then([this](TOnlineAsyncOp<FQueryFriends>& InAsyncOp, const TArray<FOnlineAccountIdHandle>& CurrentFriendIds) mutable
+		.Then([this](TOnlineAsyncOp<FQueryFriends>& InAsyncOp, const TArray<FAccountId>& CurrentFriendIds) mutable
 		{
-			const FOnlineAccountIdHandle LocalUserId = InAsyncOp.GetParams().LocalUserId;
+			const FAccountId LocalUserId = InAsyncOp.GetParams().LocalUserId;
 			const EOS_EpicAccountId LocalUserEasId = GetEpicAccountIdChecked(LocalUserId);
 
-			TMap<FOnlineAccountIdHandle, TSharedRef<FFriend>>& FriendsList = FriendsLists.FindOrAdd(LocalUserId);
-			TArray<FOnlineAccountIdHandle> NewFriendIds;
-			TArray<FOnlineAccountIdHandle> PreviousFriendIds;
-			TArray<FOnlineAccountIdHandle> UpdatedFriendIds; // TODO:  More granular updates like what changed?
+			TMap<FAccountId, TSharedRef<FFriend>>& FriendsList = FriendsLists.FindOrAdd(LocalUserId);
+			TArray<FAccountId> NewFriendIds;
+			TArray<FAccountId> PreviousFriendIds;
+			TArray<FAccountId> UpdatedFriendIds; // TODO:  More granular updates like what changed?
 			FriendsList.GenerateKeyArray(PreviousFriendIds);
 				
-			for (const FOnlineAccountIdHandle& FriendId : CurrentFriendIds)
+			for (const FAccountId& FriendId : CurrentFriendIds)
 			{
 				const EOS_EpicAccountId FriendEasId = GetEpicAccountIdChecked(FriendId);
 
@@ -163,7 +163,7 @@ TOnlineAsyncOpHandle<FQueryFriends> FFriendsEOS::QueryFriends(FQueryFriends::Par
 
 			// TODO:  Delegates
 
-			for (const FOnlineAccountIdHandle& AccountId : PreviousFriendIds)
+			for (const FAccountId& AccountId : PreviousFriendIds)
 			{
 				UE_LOG(LogTemp, Warning, TEXT("QueryFriendsComplete: %s is no longer a friend, removing"), *ToLogString(AccountId));
 				FriendsList.Remove(AccountId);
@@ -191,7 +191,7 @@ TOnlineAsyncOpHandle<FQueryFriends> FFriendsEOS::QueryFriends(FQueryFriends::Par
 
 TOnlineResult<FGetFriends> FFriendsEOS::GetFriends(FGetFriends::Params&& Params)
 {
-	if (TMap<FOnlineAccountIdHandle, TSharedRef<FFriend>>* FriendsList = FriendsLists.Find(Params.LocalUserId))
+	if (TMap<FAccountId, TSharedRef<FFriend>>* FriendsList = FriendsLists.Find(Params.LocalUserId))
 	{
 		FGetFriends::Result Result;
 		FriendsList->GenerateValueArray(Result.Friends);
@@ -245,16 +245,16 @@ TOnlineAsyncOpHandle<FAddFriend> FFriendsEOS::AddFriend(FAddFriend::Params&& InP
 	return Op->GetHandle();
 }
 
-void FFriendsEOS::OnEOSFriendsUpdate(FOnlineAccountIdHandle LocalUserId, FOnlineAccountIdHandle FriendUserId, EOS_EFriendsStatus PreviousStatus, EOS_EFriendsStatus CurrentStatus)
+void FFriendsEOS::OnEOSFriendsUpdate(FAccountId LocalUserId, FAccountId FriendUserId, EOS_EFriendsStatus PreviousStatus, EOS_EFriendsStatus CurrentStatus)
 {
 	// TODO
 	UE_LOG(LogTemp, Warning, TEXT("OnEOSFriendsUpdate: LocalUserId=[%s] FriendUserId=[%s] PreviousStatus=[%s] CurrentStatus=[%s]"), *ToLogString(LocalUserId), *ToLogString(FriendUserId), LexToString(PreviousStatus), LexToString(CurrentStatus));
 
-	TMap<FOnlineAccountIdHandle, TSharedRef<FFriend>>& FriendsList = FriendsLists.FindOrAdd(LocalUserId);
+	TMap<FAccountId, TSharedRef<FFriend>>& FriendsList = FriendsLists.FindOrAdd(LocalUserId);
 	bool bAnyChanges = false;
-	TArray<FOnlineAccountIdHandle, TInlineAllocator<1>> AddedFriends;
-	TArray<FOnlineAccountIdHandle, TInlineAllocator<1>> RemovedFriends;
-	TArray<FOnlineAccountIdHandle, TInlineAllocator<1>> UpdatedFriends;
+	TArray<FAccountId, TInlineAllocator<1>> AddedFriends;
+	TArray<FAccountId, TInlineAllocator<1>> RemovedFriends;
+	TArray<FAccountId, TInlineAllocator<1>> UpdatedFriends;
 	if (CurrentStatus == EOS_EFriendsStatus::EOS_FS_NotFriends)
 	{
 		bAnyChanges = FriendsList.Remove(FriendUserId) > 0;
diff --git a/Engine/Plugins/Online/Experimental/OnlineServicesEOS/Source/Private/Online/OnlineIdEOS.cpp b/Engine/Plugins/Online/Experimental/OnlineServicesEOS/Source/Private/Online/OnlineIdEOS.cpp
--- a/Engine/Plugins/Online/Experimental/OnlineServicesEOS/Source/Private/Online/OnlineIdEOS.cpp
+++ b/Engine/Plugins/Online/Experimental/OnlineServicesEOS/Source/Private/Online/OnlineIdEOS.cpp
@@ -9,7 +9,7 @@
 
 namespace UE::Online {
 
-FString FOnlineAccountIdRegistryEOS::ToLogString(const FOnlineAccountIdHandle& Handle) const
+FString FOnlineAccountIdRegistryEOS::ToLogString(const FAccountId& Handle) const
 {
 	FString Result;
 	if (ValidateOnlineId(Handle))
@@ -35,7 +35,7 @@ ENUM_CLASS_FLAGS(EEOSAccountIdElements);
 const uint8 OnlineIdEOSUtf8BufferLength = 32;
 const uint8 OnlineIdEOSHexBufferLength = 16;
 
-TArray<uint8> FOnlineAccountIdRegistryEOS::ToReplicationData(const FOnlineAccountIdHandle& Handle) const
+TArray<uint8> FOnlineAccountIdRegistryEOS::ToReplicationData(const FAccountId& Handle) const
 {
 	TArray<uint8> ReplicationData;
 	if (ValidateOnlineId(Handle))
@@ -87,7 +87,7 @@ TArray<uint8> FOnlineAccountIdRegistryEOS::ToReplicationData(const FOnlineAccoun
 	return ReplicationData;
 }
 
-FOnlineAccountIdHandle FOnlineAccountIdRegistryEOS::FromReplicationData(const TArray<uint8>& ReplicationData)
+FAccountId FOnlineAccountIdRegistryEOS::FromReplicationData(const TArray<uint8>& ReplicationData)
 {
 	const EEOSAccountIdElements Elements = EEOSAccountIdElements(ReplicationData[0]);
 	const int EasHexBufferLength = EnumHasAnyFlags(Elements, EEOSAccountIdElements::EAS) ? OnlineIdEOSHexBufferLength : 0;
@@ -112,13 +112,13 @@ FOnlineAccountIdHandle FOnlineAccountIdRegistryEOS::FromReplicationData(const TA
 	return FindOrAddAccountId(EpicAccountId, ProductUserId);
 }
 
-FOnlineAccountIdHandle FOnlineAccountIdRegistryEOS::FindAccountId(const EOS_ProductUserId ProductUserId) const
+FAccountId FOnlineAccountIdRegistryEOS::FindAccountId(const EOS_ProductUserId ProductUserId) const
 {
-	FOnlineAccountIdHandle Handle;
+	FAccountId Handle;
 	if (EOS_ProductUserId_IsValid(ProductUserId) == EOS_TRUE)
 	{
 		const FReadScopeLock ReadLock(Lock);
-		if (const FOnlineAccountIdHandle* Found = ProductUserIdToHandle.Find(ProductUserId))
+		if (const FAccountId* Found = ProductUserIdToHandle.Find(ProductUserId))
 		{
 			Handle = *Found;
 		}
@@ -126,18 +126,18 @@ FOnlineAccountIdHandle FOnlineAccountIdRegistryEOS::FindAccountId(const EOS_Prod
 	return Handle;
 }
 
-EOS_ProductUserId FOnlineAccountIdRegistryEOS::GetProductUserId(const FOnlineAccountIdHandle& Handle) const
+EOS_ProductUserId FOnlineAccountIdRegistryEOS::GetProductUserId(const FAccountId& Handle) const
 {
 	return GetAccountIdData(Handle).ProductUserId;
 }
 
-FOnlineAccountIdHandle FOnlineAccountIdRegistryEOS::FindAccountId(const EOS_EpicAccountId EpicAccountId) const
+FAccountId FOnlineAccountIdRegistryEOS::FindAccountId(const EOS_EpicAccountId EpicAccountId) const
 {
-	FOnlineAccountIdHandle Handle;
+	FAccountId Handle;
 	if (EOS_EpicAccountId_IsValid(EpicAccountId) == EOS_TRUE)
 	{
 		const FReadScopeLock ReadLock(Lock);
-		if (const FOnlineAccountIdHandle* Found = EpicAccountIdToHandle.Find(EpicAccountId))
+		if (const FAccountId* Found = EpicAccountIdToHandle.Find(EpicAccountId))
 		{
 			Handle = *Found;
 		}
@@ -145,7 +145,7 @@ FOnlineAccountIdHandle FOnlineAccountIdRegistryEOS::FindAccountId(const EOS_Epic
 	return Handle;
 }
 
-FOnlineAccountIdDataEOS FOnlineAccountIdRegistryEOS::GetAccountIdData(const FOnlineAccountIdHandle& Handle) const
+FOnlineAccountIdDataEOS FOnlineAccountIdRegistryEOS::GetAccountIdData(const FAccountId& Handle) const
 {
 	if (ValidateOnlineId(Handle))
 	{
@@ -161,9 +161,9 @@ FOnlineAccountIdRegistryEOS& FOnlineAccountIdRegistryEOS::Get()
 	return Instance;
 }
 
-FOnlineAccountIdHandle FOnlineAccountIdRegistryEOS::FindOrAddAccountId(const EOS_EpicAccountId InEpicAccountId, const EOS_ProductUserId InProductUserId)
+FAccountId FOnlineAccountIdRegistryEOS::FindOrAddAccountId(const EOS_EpicAccountId InEpicAccountId, const EOS_ProductUserId InProductUserId)
 {
-	FOnlineAccountIdHandle Result;
+	FAccountId Result;
 	const bool bInEpicAccountIdValid = EOS_EpicAccountId_IsValid(InEpicAccountId) == EOS_TRUE;
 	const bool bInProductUserIdValid = EOS_ProductUserId_IsValid(InProductUserId) == EOS_TRUE;
 	if (ensure(bInEpicAccountIdValid || bInProductUserIdValid))
@@ -173,12 +173,12 @@ FOnlineAccountIdHandle FOnlineAccountIdRegistryEOS::FindOrAddAccountId(const EOS
 
 		auto FindExisting = [this, InEpicAccountId, InProductUserId, bInEpicAccountIdValid, bInProductUserIdValid, &bUpdateEpicAccountId, &bUpdateProductUserId]()
 		{
-			FOnlineAccountIdHandle Result;
-			if (const FOnlineAccountIdHandle* FoundEas = bInEpicAccountIdValid ? EpicAccountIdToHandle.Find(InEpicAccountId) : nullptr)
+			FAccountId Result;
+			if (const FAccountId* FoundEas = bInEpicAccountIdValid ? EpicAccountIdToHandle.Find(InEpicAccountId) : nullptr)
 			{
 				Result = *FoundEas;
 			}
-			else if (const FOnlineAccountIdHandle* FoundProd = bInProductUserIdValid ? ProductUserIdToHandle.Find(InProductUserId) : nullptr)
+			else if (const FAccountId* FoundProd = bInProductUserIdValid ? ProductUserIdToHandle.Find(InProductUserId) : nullptr)
 			{
 				Result = *FoundProd;
 			}
@@ -218,7 +218,7 @@ FOnlineAccountIdHandle FOnlineAccountIdRegistryEOS::FindOrAddAccountId(const EOS
 				NewAccountIdData.EpicAccountId = InEpicAccountId;
 				NewAccountIdData.ProductUserId = InProductUserId;
 
-				Result = FOnlineAccountIdHandle(EOnlineServices::Epic, AccountIdData.Num());
+				Result = FAccountId(EOnlineServices::Epic, AccountIdData.Num());
 
 				if (bInEpicAccountIdValid)
 				{
@@ -253,26 +253,26 @@ FOnlineAccountIdHandle FOnlineAccountIdRegistryEOS::FindOrAddAccountId(const EOS
 	return Result;
 }
 
-EOS_EpicAccountId GetEpicAccountId(const FOnlineAccountIdHandle& Handle)
+EOS_EpicAccountId GetEpicAccountId(const FAccountId& Handle)
 {
 	return FOnlineAccountIdRegistryEOS::Get().GetAccountIdData(Handle).EpicAccountId;
 }
 
-EOS_EpicAccountId GetEpicAccountIdChecked(const FOnlineAccountIdHandle& Handle)
+EOS_EpicAccountId GetEpicAccountIdChecked(const FAccountId& Handle)
 {
 	EOS_EpicAccountId EpicAccountId = GetEpicAccountId(Handle);
 	check(EOS_EpicAccountId_IsValid(EpicAccountId) == EOS_TRUE);
 	return EpicAccountId;
 }
 
-FOnlineAccountIdHandle FindAccountId(const EOS_EpicAccountId EpicAccountId)
+FAccountId FindAccountId(const EOS_EpicAccountId EpicAccountId)
 {
 	return FOnlineAccountIdRegistryEOS::Get().FindAccountId(EpicAccountId);
 }
 
-FOnlineAccountIdHandle FindAccountIdChecked(const EOS_EpicAccountId EpicAccountId)
+FAccountId FindAccountIdChecked(const EOS_EpicAccountId EpicAccountId)
 {
-	FOnlineAccountIdHandle Result = FindAccountId(EpicAccountId);
+	FAccountId Result = FindAccountId(EpicAccountId);
 	check(Result.IsValid());
 	return Result;
 }
diff --git a/Engine/Plugins/Online/Experimental/OnlineServicesEOS/Source/Private/Online/OnlineIdEOS.h b/Engine/Plugins/Online/Experimental/OnlineServicesEOS/Source/Private/Online/OnlineIdEOS.h
--- a/Engine/Plugins/Online/Experimental/OnlineServicesEOS/Source/Private/Online/OnlineIdEOS.h
+++ b/Engine/Plugins/Online/Experimental/OnlineServicesEOS/Source/Private/Online/OnlineIdEOS.h
@@ -24,20 +24,20 @@ public:
 	virtual ~FOnlineAccountIdRegistryEOS() = default;
 
 	// Begin IOnlineAccountIdRegistry
-	virtual FString ToLogString(const FOnlineAccountIdHandle& Handle) const override;
-	virtual TArray<uint8> ToReplicationData(const FOnlineAccountIdHandle& Handle) const override;
-	virtual FOnlineAccountIdHandle FromReplicationData(const TArray<uint8>& ReplicationString) override;
+	virtual FString ToLogString(const FAccountId& Handle) const override;
+	virtual TArray<uint8> ToReplicationData(const FAccountId& Handle) const override;
+	virtual FAccountId FromReplicationData(const TArray<uint8>& ReplicationString) override;
 	// End IOnlineAccountIdRegistry
 
 	// Begin IOnlineAccountRegistryEOSGS
-	virtual FOnlineAccountIdHandle FindAccountId(const EOS_ProductUserId ProductUserId) const override;
-	virtual EOS_ProductUserId GetProductUserId(const FOnlineAccountIdHandle& Handle) const override;
+	virtual FAccountId FindAccountId(const EOS_ProductUserId ProductUserId) const override;
+	virtual EOS_ProductUserId GetProductUserId(const FAccountId& Handle) const override;
 	// End IOnlineAccountRegistryEOSGS
 
-	FOnlineAccountIdHandle FindAccountId(const EOS_EpicAccountId EpicAccountId) const;
+	FAccountId FindAccountId(const EOS_EpicAccountId EpicAccountId) const;
 
 	// Return copies as it is not thread safe to return pointers/references to array elements, in case the array is grown+relocated on another thread.
-	FOnlineAccountIdDataEOS GetAccountIdData(const FOnlineAccountIdHandle& Handle) const;
+	FOnlineAccountIdDataEOS GetAccountIdData(const FAccountId& Handle) const;
 
 	static FOnlineAccountIdRegistryEOS& Get();
 
@@ -45,19 +45,19 @@ private:
 	friend class FAuthEOS;
 	friend class FOnlineServicesEOSModule;
 	friend FAccountIdReplicationTest;
-	FOnlineAccountIdHandle FindOrAddAccountId(const EOS_EpicAccountId EpicAccountId, const EOS_ProductUserId ProductUserId);
+	FAccountId FindOrAddAccountId(const EOS_EpicAccountId EpicAccountId, const EOS_ProductUserId ProductUserId);
 
 	mutable FRWLock Lock;
 
 	TArray<FOnlineAccountIdDataEOS> AccountIdData; // Actual container for the info, indexed by the handle
 
-	TMap<EOS_EpicAccountId, FOnlineAccountIdHandle> EpicAccountIdToHandle; // Map of EOS_EpicAccountId to the associated handle.
-	TMap<EOS_ProductUserId, FOnlineAccountIdHandle> ProductUserIdToHandle; // Map of EOS_ProductUserId to the associated handle.
+	TMap<EOS_EpicAccountId, FAccountId> EpicAccountIdToHandle; // Map of EOS_EpicAccountId to the associated handle.
+	TMap<EOS_ProductUserId, FAccountId> ProductUserIdToHandle; // Map of EOS_ProductUserId to the associated handle.
 };
 
-EOS_EpicAccountId GetEpicAccountId(const FOnlineAccountIdHandle& Handle);
-EOS_EpicAccountId GetEpicAccountIdChecked(const FOnlineAccountIdHandle& Handle);
-FOnlineAccountIdHandle FindAccountId(const EOS_EpicAccountId EpicAccountId);
-FOnlineAccountIdHandle FindAccountIdChecked(const EOS_EpicAccountId EpicAccountId);
+EOS_EpicAccountId GetEpicAccountId(const FAccountId& Handle);
+EOS_EpicAccountId GetEpicAccountIdChecked(const FAccountId& Handle);
+FAccountId FindAccountId(const EOS_EpicAccountId EpicAccountId);
+FAccountId FindAccountIdChecked(const EOS_EpicAccountId EpicAccountId);
 
 } /* namespace UE::Online */
diff --git a/Engine/Plugins/Online/Experimental/OnlineServicesEOS/Source/Private/Online/PresenceEOS.cpp b/Engine/Plugins/Online/Experimental/OnlineServicesEOS/Source/Private/Online/PresenceEOS.cpp
--- a/Engine/Plugins/Online/Experimental/OnlineServicesEOS/Source/Private/Online/PresenceEOS.cpp
+++ b/Engine/Plugins/Online/Experimental/OnlineServicesEOS/Source/Private/Online/PresenceEOS.cpp
@@ -78,10 +78,10 @@ void FPresenceEOS::Initialize()
 	NotifyPresenceChangedNotificationId = EOS_Presence_AddNotifyOnPresenceChanged(PresenceHandle, &Options, this, [](const EOS_Presence_PresenceChangedCallbackInfo* Data)
 	{
 		FPresenceEOS* This = reinterpret_cast<FPresenceEOS*>(Data->ClientData);
-		const FOnlineAccountIdHandle LocalUserId = FindAccountIdChecked(Data->LocalUserId);
+		const FAccountId LocalUserId = FindAccountIdChecked(Data->LocalUserId);
 
 		This->Services.Get<FAuthEOS>()->ResolveAccountId(LocalUserId, Data->PresenceUserId)
-		.Next([This, LocalUserId](const FOnlineAccountIdHandle& PresenceUserId)
+		.Next([This, LocalUserId](const FAccountId& PresenceUserId)
 		{
 			UE_LOG(LogTemp, Warning, TEXT("OnEOSPresenceUpdate: LocalUserId=[%s] PresenceUserId=[%s]"), *ToLogString(LocalUserId), *ToLogString(PresenceUserId));
 			This->UpdateUserPresence(LocalUserId, PresenceUserId);
@@ -154,7 +154,7 @@ TOnlineAsyncOpHandle<FQueryPresence> FPresenceEOS::QueryPresence(FQueryPresence:
 
 TOnlineResult<FGetCachedPresence> FPresenceEOS::GetCachedPresence(FGetCachedPresence::Params&& Params)
 {
-	if (TMap<FOnlineAccountIdHandle, TSharedRef<FUserPresence>>* PresenceList = PresenceLists.Find(Params.LocalUserId))
+	if (TMap<FAccountId, TSharedRef<FUserPresence>>* PresenceList = PresenceLists.Find(Params.LocalUserId))
 	{
 		TSharedRef<FUserPresence>* PresencePtr = PresenceList->Find(Params.TargetUserId);
 		if (PresencePtr)
@@ -548,9 +548,9 @@ TOnlineAsyncOpHandle<FPartialUpdatePresence> FPresenceEOS::PartialUpdatePresence
 }
 
 /** Get a user's presence, creating entries if missing */
-TSharedRef<FUserPresence> FPresenceEOS::FindOrCreatePresence(FOnlineAccountIdHandle LocalUserId, FOnlineAccountIdHandle PresenceUserId)
+TSharedRef<FUserPresence> FPresenceEOS::FindOrCreatePresence(FAccountId LocalUserId, FAccountId PresenceUserId)
 {
-	TMap<FOnlineAccountIdHandle, TSharedRef<FUserPresence>>& LocalUserPresenceList = PresenceLists.FindOrAdd(LocalUserId);
+	TMap<FAccountId, TSharedRef<FUserPresence>>& LocalUserPresenceList = PresenceLists.FindOrAdd(LocalUserId);
 	if (const TSharedRef<FUserPresence>* const ExistingPresence = LocalUserPresenceList.Find(PresenceUserId))
 	{
 		return *ExistingPresence;
@@ -562,7 +562,7 @@ TSharedRef<FUserPresence> FPresenceEOS::FindOrCreatePresence(FOnlineAccountIdHan
 	return UserPresence;
 }
 
-void FPresenceEOS::UpdateUserPresence(FOnlineAccountIdHandle LocalUserId, FOnlineAccountIdHandle PresenceUserId)
+void FPresenceEOS::UpdateUserPresence(FAccountId LocalUserId, FAccountId PresenceUserId)
 {
 	bool bPresenceHasChanged = false;
 	TSharedRef<FUserPresence> UserPresence = FindOrCreatePresence(LocalUserId, PresenceUserId);
diff --git a/Engine/Plugins/Online/Experimental/OnlineServicesEOS/Source/Private/Online/Tests/OnlineIdEOSTests.cpp b/Engine/Plugins/Online/Experimental/OnlineServicesEOS/Source/Private/Online/Tests/OnlineIdEOSTests.cpp
--- a/Engine/Plugins/Online/Experimental/OnlineServicesEOS/Source/Private/Online/Tests/OnlineIdEOSTests.cpp
+++ b/Engine/Plugins/Online/Experimental/OnlineServicesEOS/Source/Private/Online/Tests/OnlineIdEOSTests.cpp
@@ -21,25 +21,25 @@ bool FAccountIdReplicationTest::RunTest(const FString& Parameters)
 
 		{
 			FOnlineAccountIdRegistryEOS Registry;
-			const FOnlineAccountIdHandle Handle = Registry.FindOrAddAccountId(EasId, EosId);
+			const FAccountId Handle = Registry.FindOrAddAccountId(EasId, EosId);
 			TArray<uint8> RepData = Registry.ToReplicationData(Handle);
-			const FOnlineAccountIdHandle Handle2 = Registry.FromReplicationData(RepData);
+			const FAccountId Handle2 = Registry.FromReplicationData(RepData);
 			UTEST_EQUAL(TEXT(""), Handle, Handle2);
 		}
 
 		{
 			FOnlineAccountIdRegistryEOS Registry;
-			const FOnlineAccountIdHandle Handle = Registry.FindOrAddAccountId(EasId, nullptr);
+			const FAccountId Handle = Registry.FindOrAddAccountId(EasId, nullptr);
 			TArray<uint8> RepData = Registry.ToReplicationData(Handle);
-			const FOnlineAccountIdHandle Handle2 = Registry.FromReplicationData(RepData);
+			const FAccountId Handle2 = Registry.FromReplicationData(RepData);
 			UTEST_EQUAL(TEXT(""), Handle, Handle2);
 		}
 
 		{
 			FOnlineAccountIdRegistryEOS Registry;
-			const FOnlineAccountIdHandle Handle = Registry.FindOrAddAccountId(nullptr, EosId);
+			const FAccountId Handle = Registry.FindOrAddAccountId(nullptr, EosId);
 			TArray<uint8> RepData = Registry.ToReplicationData(Handle);
-			const FOnlineAccountIdHandle Handle2 = Registry.FromReplicationData(RepData);
+			const FAccountId Handle2 = Registry.FromReplicationData(RepData);
 			UTEST_EQUAL(TEXT(""), Handle, Handle2);
 		}
 	}
diff --git a/Engine/Plugins/Online/Experimental/OnlineServicesEOS/Source/Private/Online/UserInfoEOS.cpp b/Engine/Plugins/Online/Experimental/OnlineServicesEOS/Source/Private/Online/UserInfoEOS.cpp
--- a/Engine/Plugins/Online/Experimental/OnlineServicesEOS/Source/Private/Online/UserInfoEOS.cpp
+++ b/Engine/Plugins/Online/Experimental/OnlineServicesEOS/Source/Private/Online/UserInfoEOS.cpp
@@ -36,7 +36,7 @@ TOnlineAsyncOpHandle<FQueryUserInfo> FUserInfoEOS::QueryUserInfo(FQueryUserInfo:
 		return Op->GetHandle();
 	}
 
-	for (const FOnlineAccountIdHandle TargetUserId : Params.UserIds)
+	for (const FAccountId TargetUserId : Params.UserIds)
 	{
 		Op->Then([this, TargetUserId](TOnlineAsyncOp<FQueryUserInfo>& Op, TPromise<const EOS_UserInfo_QueryUserInfoCallbackInfo*>&& Promise)
 		{
diff --git a/Engine/Plugins/Online/Experimental/OnlineServicesEOS/Source/Public/Online/AuthEOS.h b/Engine/Plugins/Online/Experimental/OnlineServicesEOS/Source/Public/Online/AuthEOS.h
--- a/Engine/Plugins/Online/Experimental/OnlineServicesEOS/Source/Public/Online/AuthEOS.h
+++ b/Engine/Plugins/Online/Experimental/OnlineServicesEOS/Source/Public/Online/AuthEOS.h
@@ -29,22 +29,22 @@ public:
 	// End IAuth
 
 	// Begin FAuthEOSGS
-	virtual TFuture<FOnlineAccountIdHandle> ResolveAccountId(const FOnlineAccountIdHandle& LocalUserId, const EOS_ProductUserId ProductUserId) override;
-	virtual TFuture<TArray<FOnlineAccountIdHandle>> ResolveAccountIds(const FOnlineAccountIdHandle& LocalUserId, const TArray<EOS_ProductUserId>& ProductUserIds) override;
-	virtual TFunction<TFuture<FOnlineAccountIdHandle>(FOnlineAsyncOp& InAsyncOp, const EOS_ProductUserId& ProductUserId)> ResolveProductIdFn() override;
-	virtual TFunction<TFuture<TArray<FOnlineAccountIdHandle>>(FOnlineAsyncOp& InAsyncOp, const TArray<EOS_ProductUserId>& ProductUserIds)> ResolveProductIdsFn() override;
+	virtual TFuture<FAccountId> ResolveAccountId(const FAccountId& LocalUserId, const EOS_ProductUserId ProductUserId) override;
+	virtual TFuture<TArray<FAccountId>> ResolveAccountIds(const FAccountId& LocalUserId, const TArray<EOS_ProductUserId>& ProductUserIds) override;
+	virtual TFunction<TFuture<FAccountId>(FOnlineAsyncOp& InAsyncOp, const EOS_ProductUserId& ProductUserId)> ResolveProductIdFn() override;
+	virtual TFunction<TFuture<TArray<FAccountId>>(FOnlineAsyncOp& InAsyncOp, const TArray<EOS_ProductUserId>& ProductUserIds)> ResolveProductIdsFn() override;
 	// End FAuthEOSGS
 
-	TFuture<FOnlineAccountIdHandle> ResolveAccountId(const FOnlineAccountIdHandle& LocalUserId, const EOS_EpicAccountId EpicAccountId);
-	TFuture<TArray<FOnlineAccountIdHandle>> ResolveAccountIds(const FOnlineAccountIdHandle& LocalUserId, const TArray<EOS_EpicAccountId>& EpicAccountIds);
-	TFunction<TFuture<FOnlineAccountIdHandle>(FOnlineAsyncOp& InAsyncOp, const EOS_EpicAccountId& EpicAccountId)> ResolveEpicIdFn();
-	TFunction<TFuture<TArray<FOnlineAccountIdHandle>>(FOnlineAsyncOp& InAsyncOp, const TArray<EOS_EpicAccountId>& EpicAccountIds)> ResolveEpicIdsFn();
+	TFuture<FAccountId> ResolveAccountId(const FAccountId& LocalUserId, const EOS_EpicAccountId EpicAccountId);
+	TFuture<TArray<FAccountId>> ResolveAccountIds(const FAccountId& LocalUserId, const TArray<EOS_EpicAccountId>& EpicAccountIds);
+	TFunction<TFuture<FAccountId>(FOnlineAsyncOp& InAsyncOp, const EOS_EpicAccountId& EpicAccountId)> ResolveEpicIdFn();
+	TFunction<TFuture<TArray<FAccountId>>(FOnlineAsyncOp& InAsyncOp, const TArray<EOS_EpicAccountId>& EpicAccountIds)> ResolveEpicIdsFn();
 
 protected:
 	void ProcessSuccessfulLogin(TOnlineAsyncOp<FAuthLogin>& InAsyncOp);
-	void OnEASLoginStatusChanged(FOnlineAccountIdHandle LocalUserId, ELoginStatus PreviousStatus, ELoginStatus CurrentStatus);
+	void OnEASLoginStatusChanged(FAccountId LocalUserId, ELoginStatus PreviousStatus, ELoginStatus CurrentStatus);
 
-	static FOnlineAccountIdHandle CreateAccountId(const EOS_EpicAccountId EpicAccountId, const EOS_ProductUserId ProductUserId);
+	static FAccountId CreateAccountId(const EOS_EpicAccountId EpicAccountId, const EOS_ProductUserId ProductUserId);
 
 	EOS_HUserInfo UserInfoHandle = nullptr;
 };
diff --git a/Engine/Plugins/Online/Experimental/OnlineServicesEOS/Source/Public/Online/CommerceEOS.h b/Engine/Plugins/Online/Experimental/OnlineServicesEOS/Source/Public/Online/CommerceEOS.h
--- a/Engine/Plugins/Online/Experimental/OnlineServicesEOS/Source/Public/Online/CommerceEOS.h
+++ b/Engine/Plugins/Online/Experimental/OnlineServicesEOS/Source/Public/Online/CommerceEOS.h
@@ -41,8 +41,8 @@ private:
 	void EOSEntitlementToOssEntitlement(FEntitlement& OutEntitlement, EOS_Ecom_Entitlement* EosEntitlement);
 
 	EOS_HEcom EcomHandle;
-	TMap<FOnlineAccountIdHandle, TArray<FEntitlement>> CachedEntitlements;
-	TMap<FOnlineAccountIdHandle, TArray<FOffer>> CachedOffers;
+	TMap<FAccountId, TArray<FEntitlement>> CachedEntitlements;
+	TMap<FAccountId, TArray<FOffer>> CachedOffers;
 };
 
 /* UE::Online */
diff --git a/Engine/Plugins/Online/Experimental/OnlineServicesEOS/Source/Public/Online/FriendsEOS.h b/Engine/Plugins/Online/Experimental/OnlineServicesEOS/Source/Public/Online/FriendsEOS.h
--- a/Engine/Plugins/Online/Experimental/OnlineServicesEOS/Source/Public/Online/FriendsEOS.h
+++ b/Engine/Plugins/Online/Experimental/OnlineServicesEOS/Source/Public/Online/FriendsEOS.h
@@ -27,9 +27,9 @@ public:
 	virtual TOnlineAsyncOpHandle<FAddFriend> AddFriend(FAddFriend::Params&& Params) override;
 
 protected:
-	void OnEOSFriendsUpdate(FOnlineAccountIdHandle LocalUserId, FOnlineAccountIdHandle FriendUserId, EOS_EFriendsStatus PreviousStatus, EOS_EFriendsStatus CurrentStatus);
+	void OnEOSFriendsUpdate(FAccountId LocalUserId, FAccountId FriendUserId, EOS_EFriendsStatus PreviousStatus, EOS_EFriendsStatus CurrentStatus);
 protected:
-	TMap<FOnlineAccountIdHandle, TMap<FOnlineAccountIdHandle, TSharedRef<FFriend>>> FriendsLists;
+	TMap<FAccountId, TMap<FAccountId, TSharedRef<FFriend>>> FriendsLists;
 	EOS_HFriends FriendsHandle = nullptr;
 
 	EOS_NotificationId NotifyFriendsUpdateNotificationId = 0;
diff --git a/Engine/Plugins/Online/Experimental/OnlineServicesEOS/Source/Public/Online/PresenceEOS.h b/Engine/Plugins/Online/Experimental/OnlineServicesEOS/Source/Public/Online/PresenceEOS.h
--- a/Engine/Plugins/Online/Experimental/OnlineServicesEOS/Source/Public/Online/PresenceEOS.h
+++ b/Engine/Plugins/Online/Experimental/OnlineServicesEOS/Source/Public/Online/PresenceEOS.h
@@ -29,13 +29,13 @@ public:
 
 protected:
 	/** Get a user's presence, creating entries if missing */
-	TSharedRef<FUserPresence> FindOrCreatePresence(FOnlineAccountIdHandle LocalUserId, FOnlineAccountIdHandle PresenceUserId);
+	TSharedRef<FUserPresence> FindOrCreatePresence(FAccountId LocalUserId, FAccountId PresenceUserId);
 	/** Update a user's presence from EOS's current value */
-	void UpdateUserPresence(FOnlineAccountIdHandle LocalUserId, FOnlineAccountIdHandle PresenceUserId);
+	void UpdateUserPresence(FAccountId LocalUserId, FAccountId PresenceUserId);
 protected:
 	EOS_HPresence PresenceHandle = nullptr;
 
-	TMap<FOnlineAccountIdHandle, TMap<FOnlineAccountIdHandle, TSharedRef<FUserPresence>>> PresenceLists;
+	TMap<FAccountId, TMap<FAccountId, TSharedRef<FUserPresence>>> PresenceLists;
 	EOS_NotificationId NotifyPresenceChangedNotificationId = 0;
 };
 
diff --git a/Engine/Plugins/Online/Experimental/OnlineServicesEOSGS/Source/Private/Online/AchievementsEOSGS.cpp b/Engine/Plugins/Online/Experimental/OnlineServicesEOSGS/Source/Private/Online/AchievementsEOSGS.cpp
--- a/Engine/Plugins/Online/Experimental/OnlineServicesEOSGS/Source/Private/Online/AchievementsEOSGS.cpp
+++ b/Engine/Plugins/Online/Experimental/OnlineServicesEOSGS/Source/Private/Online/AchievementsEOSGS.cpp
@@ -33,7 +33,7 @@ void FAchievementsEOSGS::Initialize()
 	{
 		FAchievementsEOSGS* This = reinterpret_cast<FAchievementsEOSGS*>(Data->ClientData);
 
-		const FOnlineAccountIdHandle LocalUserId = FindAccountId(Data->UserId);
+		const FAccountId LocalUserId = FindAccountId(Data->UserId);
 		if (LocalUserId.IsValid())
 		{
 			const FString AchievementId = UTF8_TO_TCHAR(Data->AchievementId);
diff --git a/Engine/Plugins/Online/Experimental/OnlineServicesEOSGS/Source/Private/Online/AuthEOSGS.cpp b/Engine/Plugins/Online/Experimental/OnlineServicesEOSGS/Source/Private/Online/AuthEOSGS.cpp
--- a/Engine/Plugins/Online/Experimental/OnlineServicesEOSGS/Source/Private/Online/AuthEOSGS.cpp
+++ b/Engine/Plugins/Online/Experimental/OnlineServicesEOSGS/Source/Private/Online/AuthEOSGS.cpp
@@ -473,7 +473,7 @@ TSharedPtr<FAccountInfoEOS> FAccountInfoRegistryEOS::Find(FPlatformUserId Platfo
 	return StaticCastSharedPtr<FAccountInfoEOS>(Super::Find(PlatformUserId));
 }
 
-TSharedPtr<FAccountInfoEOS> FAccountInfoRegistryEOS::Find(FOnlineAccountIdHandle AccountIdHandle) const
+TSharedPtr<FAccountInfoEOS> FAccountInfoRegistryEOS::Find(FAccountId AccountIdHandle) const
 {
 	return StaticCastSharedPtr<FAccountInfoEOS>(Super::Find(AccountIdHandle));
 }
@@ -510,7 +510,7 @@ void FAccountInfoRegistryEOS::Register(const TSharedRef<FAccountInfoEOS>& Accoun
 	DoRegister(AccountInfoEOS);
 }
 
-void FAccountInfoRegistryEOS::Unregister(FOnlineAccountIdHandle AccountId)
+void FAccountInfoRegistryEOS::Unregister(FAccountId AccountId)
 {
 	if (TSharedPtr<FAccountInfoEOS> AccountInfoEOS = Find(AccountId))
 	{
@@ -946,43 +946,43 @@ TOnlineAsyncOpHandle<FAuthEndVerifiedAuthSession> FAuthEOSGS::EndVerifiedAuthSes
 	return Operation->GetHandle();
 }
 
-TFuture<FOnlineAccountIdHandle> FAuthEOSGS::ResolveAccountId(const FOnlineAccountIdHandle& LocalUserId, const EOS_ProductUserId ProductUserId)
+TFuture<FAccountId> FAuthEOSGS::ResolveAccountId(const FAccountId& LocalUserId, const EOS_ProductUserId ProductUserId)
 {
-	return MakeFulfilledPromise<FOnlineAccountIdHandle>(CreateAccountId(ProductUserId)).GetFuture();
+	return MakeFulfilledPromise<FAccountId>(CreateAccountId(ProductUserId)).GetFuture();
 }
 
-TFuture<TArray<FOnlineAccountIdHandle>> FAuthEOSGS::ResolveAccountIds(const FOnlineAccountIdHandle& LocalUserId, const TArray<EOS_ProductUserId>& InProductUserIds)
+TFuture<TArray<FAccountId>> FAuthEOSGS::ResolveAccountIds(const FAccountId& LocalUserId, const TArray<EOS_ProductUserId>& InProductUserIds)
 {
-	TArray<FOnlineAccountIdHandle> AccountIds;
+	TArray<FAccountId> AccountIds;
 	AccountIds.Reserve(InProductUserIds.Num());
 	for (const EOS_ProductUserId ProductUserId : InProductUserIds)
 	{
 		AccountIds.Emplace(CreateAccountId(ProductUserId));
 	}
-	return MakeFulfilledPromise<TArray<FOnlineAccountIdHandle>>(MoveTemp(AccountIds)).GetFuture();
+	return MakeFulfilledPromise<TArray<FAccountId>>(MoveTemp(AccountIds)).GetFuture();
 }
 
-TFunction<TFuture<FOnlineAccountIdHandle>(FOnlineAsyncOp& InAsyncOp, const EOS_ProductUserId& ProductUserId)> FAuthEOSGS::ResolveProductIdFn()
+TFunction<TFuture<FAccountId>(FOnlineAsyncOp& InAsyncOp, const EOS_ProductUserId& ProductUserId)> FAuthEOSGS::ResolveProductIdFn()
 {
 	return [this](FOnlineAsyncOp& InAsyncOp, const EOS_ProductUserId& ProductUserId)
 	{
-		const FOnlineAccountIdHandle* LocalUserIdPtr = InAsyncOp.Data.Get<FOnlineAccountIdHandle>(TEXT("LocalUserId"));
+		const FAccountId* LocalUserIdPtr = InAsyncOp.Data.Get<FAccountId>(TEXT("LocalUserId"));
 		if (!ensure(LocalUserIdPtr))
 		{
-			return MakeFulfilledPromise<FOnlineAccountIdHandle>().GetFuture();
+			return MakeFulfilledPromise<FAccountId>().GetFuture();
 		}
 		return ResolveAccountId(*LocalUserIdPtr, ProductUserId);
 	};
 }
 
-TFunction<TFuture<TArray<FOnlineAccountIdHandle>>(FOnlineAsyncOp& InAsyncOp, const TArray<EOS_ProductUserId>& ProductUserIds)> FAuthEOSGS::ResolveProductIdsFn()
+TFunction<TFuture<TArray<FAccountId>>(FOnlineAsyncOp& InAsyncOp, const TArray<EOS_ProductUserId>& ProductUserIds)> FAuthEOSGS::ResolveProductIdsFn()
 {
 	return [this](FOnlineAsyncOp& InAsyncOp, const TArray<EOS_ProductUserId>& ProductUserIds)
 	{
-		const FOnlineAccountIdHandle* LocalUserIdPtr = InAsyncOp.Data.Get<FOnlineAccountIdHandle>(TEXT("LocalUserId"));
+		const FAccountId* LocalUserIdPtr = InAsyncOp.Data.Get<FAccountId>(TEXT("LocalUserId"));
 		if (!ensure(LocalUserIdPtr))
 		{
-			return MakeFulfilledPromise<TArray<FOnlineAccountIdHandle>>().GetFuture();
+			return MakeFulfilledPromise<TArray<FAccountId>>().GetFuture();
 		}
 		return ResolveAccountIds(*LocalUserIdPtr, ProductUserIds);
 	};
@@ -1540,7 +1540,7 @@ void FAuthEOSGS::InitializeConnectLoginRecoveryTimer(const TSharedRef<FAccountIn
 		DelayTime);
 }
 
-FOnlineAccountIdHandle FAuthEOSGS::CreateAccountId(const EOS_ProductUserId ProductUserId)
+FAccountId FAuthEOSGS::CreateAccountId(const EOS_ProductUserId ProductUserId)
 {
 	return FOnlineAccountIdRegistryEOSGS::Get().FindOrAddAccountId(ProductUserId);
 }
diff --git a/Engine/Plugins/Online/Experimental/OnlineServicesEOSGS/Source/Private/Online/LeaderboardsEOSGS.cpp b/Engine/Plugins/Online/Experimental/OnlineServicesEOSGS/Source/Private/Online/LeaderboardsEOSGS.cpp
--- a/Engine/Plugins/Online/Experimental/OnlineServicesEOSGS/Source/Private/Online/LeaderboardsEOSGS.cpp
+++ b/Engine/Plugins/Online/Experimental/OnlineServicesEOSGS/Source/Private/Online/LeaderboardsEOSGS.cpp
@@ -42,7 +42,7 @@ TOnlineAsyncOpHandle<FReadEntriesForUsers> FLeaderboardsEOSGS::ReadEntriesForUse
 	{
 		TArray<EOS_ProductUserId> ProductUserIds;
 		ProductUserIds.AddZeroed(InAsyncOp.GetParams().UserIds.Num());
-		for (const FOnlineAccountIdHandle& UserId : InAsyncOp.GetParams().UserIds)
+		for (const FAccountId& UserId : InAsyncOp.GetParams().UserIds)
 		{
 			ProductUserIds.Emplace(GetProductUserIdChecked(UserId));
 		}
@@ -84,7 +84,7 @@ TOnlineAsyncOpHandle<FReadEntriesForUsers> FLeaderboardsEOSGS::ReadEntriesForUse
 
 		FReadEntriesForUsers::Result Result;
 
-		for (const FOnlineAccountIdHandle& UserId : InAsyncOp.GetParams().UserIds)
+		for (const FAccountId& UserId : InAsyncOp.GetParams().UserIds)
 		{
 			EOS_Leaderboards_CopyLeaderboardUserScoreByUserIdOptions UserCopyOptions = { };
 			UserCopyOptions.ApiVersion = EOS_LEADERBOARDS_COPYLEADERBOARDUSERSCOREBYUSERID_API_LATEST;
@@ -120,7 +120,7 @@ TOnlineAsyncOpHandle<FReadEntriesForUsers> FLeaderboardsEOSGS::ReadEntriesForUse
 namespace Private
 {
 
-void QueryLeaderboardsEOS(EOS_HLeaderboards LeaderboardsHandle, const FOnlineAccountIdHandle& LocalUserId, const FString& BoardName, TPromise<const EOS_Leaderboards_OnQueryLeaderboardRanksCompleteCallbackInfo*>&& Promise)
+void QueryLeaderboardsEOS(EOS_HLeaderboards LeaderboardsHandle, const FAccountId& LocalUserId, const FString& BoardName, TPromise<const EOS_Leaderboards_OnQueryLeaderboardRanksCompleteCallbackInfo*>&& Promise)
 {
 	EOS_Leaderboards_QueryLeaderboardRanksOptions Options;
 	Options.LocalUserId = GetProductUserIdChecked(LocalUserId);
diff --git a/Engine/Plugins/Online/Experimental/OnlineServicesEOSGS/Source/Private/Online/LobbiesEOSGS.cpp b/Engine/Plugins/Online/Experimental/OnlineServicesEOSGS/Source/Private/Online/LobbiesEOSGS.cpp
--- a/Engine/Plugins/Online/Experimental/OnlineServicesEOSGS/Source/Private/Online/LobbiesEOSGS.cpp
+++ b/Engine/Plugins/Online/Experimental/OnlineServicesEOSGS/Source/Private/Online/LobbiesEOSGS.cpp
@@ -71,7 +71,7 @@ TOnlineAsyncOpHandle<FCreateLobby> FLobbiesEOSGS::CreateLobby(FCreateLobby::Para
 	auto DestroyLobbyDuringCreate = [this](
 		TOnlineAsyncOp<FCreateLobby>& InAsyncOp,
 		const TSharedPtr<FLobbyDataEOS>& LobbyData,
-		FOnlineAccountIdHandle LocalUserId,
+		FAccountId LocalUserId,
 		FOnlineError ErrorResult) -> TFuture<void>
 	{
 		FLobbiesDestroyLobbyImpl::Params DestroyLobbyParams;
@@ -186,10 +186,10 @@ TOnlineAsyncOpHandle<FCreateLobby> FLobbiesEOSGS::CreateLobby(FCreateLobby::Para
 
 		// Add owner attributes to operation data.
 		// CreatingMemberData is used to update the local lobby data and for dispatching notifications once creation has completed.
-		TMap<FOnlineAccountIdHandle, TSharedRef<FClientLobbyMemberDataChanges>> CreatingMemberData;
+		TMap<FAccountId, TSharedRef<FClientLobbyMemberDataChanges>> CreatingMemberData;
 		CreatingMemberData.Reserve(Params.LocalUsers.Num());
 		CreatingMemberData.Add(Params.LocalUserId, LobbyOwnerAttributes);
-		InAsyncOp.Data.Set<TMap<FOnlineAccountIdHandle, TSharedRef<FClientLobbyMemberDataChanges>>>(LobbyMemberChangesKeyName, MoveTemp(CreatingMemberData));
+		InAsyncOp.Data.Set<TMap<FAccountId, TSharedRef<FClientLobbyMemberDataChanges>>>(LobbyMemberChangesKeyName, MoveTemp(CreatingMemberData));
 
 		if (UserData)
 		{
@@ -284,8 +284,8 @@ TOnlineAsyncOpHandle<FCreateLobby> FLobbiesEOSGS::CreateLobby(FCreateLobby::Para
 
 		// Add member changes to the lobby changes object.
 		const TSharedRef<FClientLobbyDataChanges>& LobbyChanges = GetOpDataChecked<TSharedRef<FClientLobbyDataChanges>>(InAsyncOp, LobbyChangesKeyName);
-		const TMap<FOnlineAccountIdHandle, TSharedRef<FClientLobbyMemberDataChanges>>& MemberChanges =
-			GetOpDataChecked<TMap<FOnlineAccountIdHandle, TSharedRef<FClientLobbyMemberDataChanges>>>(InAsyncOp, LobbyMemberChangesKeyName);
+		const TMap<FAccountId, TSharedRef<FClientLobbyMemberDataChanges>>& MemberChanges =
+			GetOpDataChecked<TMap<FAccountId, TSharedRef<FClientLobbyMemberDataChanges>>>(InAsyncOp, LobbyMemberChangesKeyName);
 		LobbyChanges->MutatedMembers = MemberChanges;
 		LobbyChanges->LocalName = InAsyncOp.GetParams().LocalName;
 
@@ -447,7 +447,7 @@ TOnlineAsyncOpHandle<FLeaveLobby> FLobbiesEOSGS::LeaveLobby(FLeaveLobby::Params&
 
 		// Remove the lobby from the active list for the user.
 		// The lobby data will be cleaned up once all references are removed.
-		for (FOnlineAccountIdHandle LeavingMember : ApplyResult.LeavingLocalMembers)
+		for (FAccountId LeavingMember : ApplyResult.LeavingLocalMembers)
 		{
 			RemoveActiveLobby(LeavingMember, LobbyData.ToSharedRef());
 		}
@@ -944,7 +944,7 @@ void FLobbiesEOSGS::HandleLobbyInviteReceived(const EOS_Lobby_LobbyInviteReceive
 #endif
 
 	// Todo: Queue this like an operation.
-	const FOnlineAccountIdHandle LocalUserId = FindAccountId(Data->LocalUserId);
+	const FAccountId LocalUserId = FindAccountId(Data->LocalUserId);
 	if (LocalUserId.IsValid())
 	{
 		FLobbyInviteDataEOS::CreateFromInviteId(LobbyPrerequisites.ToSharedRef(), LobbyDataRegistry.ToSharedRef(), LocalUserId, Data->InviteId, Data->TargetUserId)
@@ -972,7 +972,7 @@ void FLobbiesEOSGS::HandleLobbyInviteReceived(const EOS_Lobby_LobbyInviteReceive
 void FLobbiesEOSGS::HandleLobbyInviteAccepted(const EOS_Lobby_LobbyInviteAcceptedCallbackInfo* Data)
 {
 	// Todo: Queue this like an operation.
-	const FOnlineAccountIdHandle LocalUserId = FindAccountId(Data->LocalUserId);
+	const FAccountId LocalUserId = FindAccountId(Data->LocalUserId);
 	if (LocalUserId.IsValid())
 	{
 		FLobbyInviteDataEOS::CreateFromInviteId(LobbyPrerequisites.ToSharedRef(), LobbyDataRegistry.ToSharedRef(), LocalUserId, Data->InviteId, Data->TargetUserId)
@@ -1007,7 +1007,7 @@ void FLobbiesEOSGS::HandleLobbyInviteAccepted(const EOS_Lobby_LobbyInviteAccepte
 void FLobbiesEOSGS::HandleJoinLobbyAccepted(const EOS_Lobby_JoinLobbyAcceptedCallbackInfo* Data)
 {
 	// Todo: Queue this like an operation.
-	const FOnlineAccountIdHandle LocalUserId = FindAccountId(Data->LocalUserId);
+	const FAccountId LocalUserId = FindAccountId(Data->LocalUserId);
 	if (LocalUserId.IsValid())
 	{
 		TDefaultErrorResultInternal<TSharedRef<FLobbyDetailsEOS>> LobbyDetailsResult = 
@@ -1154,13 +1154,13 @@ void FLobbiesEOSGS::UnregisterHandlers()
 	OnJoinLobbyAcceptedEOSEventRegistration = nullptr;
 }
 
-void FLobbiesEOSGS::AddActiveLobby(FOnlineAccountIdHandle LocalUserId, const TSharedRef<FLobbyDataEOS>& LobbyData)
+void FLobbiesEOSGS::AddActiveLobby(FAccountId LocalUserId, const TSharedRef<FLobbyDataEOS>& LobbyData)
 {
 	// Add bookkeeping for the user.
 	ActiveLobbies.FindOrAdd(LocalUserId).Add(LobbyData);
 }
 
-void FLobbiesEOSGS::RemoveActiveLobby(FOnlineAccountIdHandle LocalUserId, const TSharedRef<FLobbyDataEOS>& LobbyData)
+void FLobbiesEOSGS::RemoveActiveLobby(FAccountId LocalUserId, const TSharedRef<FLobbyDataEOS>& LobbyData)
 {
 	// Remove bookkeeping for the local user.
 	if (TSet<TSharedRef<FLobbyDataEOS>>* Lobbies = ActiveLobbies.Find(LocalUserId))
@@ -1199,7 +1199,7 @@ void FLobbiesEOSGS::RemoveActiveInvite(const TSharedRef<FLobbyInviteDataEOS>& In
 		});
 }
 
-TSharedPtr<FLobbyInviteDataEOS> FLobbiesEOSGS::GetActiveInvite(FOnlineAccountIdHandle TargetUser, FOnlineLobbyIdHandle TargetLobbyId)
+TSharedPtr<FLobbyInviteDataEOS> FLobbiesEOSGS::GetActiveInvite(FAccountId TargetUser, FOnlineLobbyIdHandle TargetLobbyId)
 {
 	const TSharedRef<FLobbyInviteDataEOS>* Result = ActiveInvites.FindOrAdd(TargetUser).Find(TargetLobbyId);
 	return Result ? *Result : TSharedPtr<FLobbyInviteDataEOS>();
@@ -1276,7 +1276,7 @@ TFuture<TDefaultErrorResult<FLobbiesJoinLobbyImpl>> FLobbiesEOSGS::JoinLobbyImpl
 			TArray<TFuture<TDefaultErrorResult<FLobbiesLeaveLobbyImpl>>> PendingMemberExits;
 			for (int32 MemberIndex = 0; MemberIndex < Params.LocalUsers.Num(); ++MemberIndex)
 			{
-				const FOnlineAccountIdHandle MemberId = Params.LocalUsers[MemberIndex].LocalUserId;
+				const FAccountId MemberId = Params.LocalUsers[MemberIndex].LocalUserId;
 
 				if (Results.Get()[MemberIndex].IsError())
 				{
@@ -1783,11 +1783,11 @@ TOnlineAsyncOpHandle<FLobbiesProcessLobbyNotificationImpl> FLobbiesEOSGS::Proces
 		}
 
 		// Get member snapshots.
-		TMap<FOnlineAccountIdHandle, TSharedRef<FClientLobbyMemberSnapshot>> LobbyMemberSnapshots;
+		TMap<FAccountId, TSharedRef<FClientLobbyMemberSnapshot>> LobbyMemberSnapshots;
 		LobbyMemberSnapshots.Reserve(Params.MutatedMembers.Num());
 		for (EOS_ProductUserId MutatedMember : Params.MutatedMembers)
 		{
-			const FOnlineAccountIdHandle MutatedMemberAccountId = FindAccountId(MutatedMember);
+			const FAccountId MutatedMemberAccountId = FindAccountId(MutatedMember);
 			if (MutatedMemberAccountId.IsValid())
 			{
 				TDefaultErrorResultInternal<TSharedRef<FClientLobbyMemberSnapshot>> MemberSnapshotResult = LobbyDetails->GetLobbyMemberSnapshot(MutatedMemberAccountId);
@@ -1801,12 +1801,12 @@ TOnlineAsyncOpHandle<FLobbiesProcessLobbyNotificationImpl> FLobbiesEOSGS::Proces
 			}
 		}
 
-		// Translate leaving members from EOS_ProductUserId to FOnlineAccountIdHandle.
-		TMap<FOnlineAccountIdHandle, ELobbyMemberLeaveReason> LeavingMemberReason;
+		// Translate leaving members from EOS_ProductUserId to FAccountId.
+		TMap<FAccountId, ELobbyMemberLeaveReason> LeavingMemberReason;
 		LeavingMemberReason.Reserve(Params.LeavingMembers.Num());
 		for (const TPair<EOS_ProductUserId, ELobbyMemberLeaveReason>& LeavingMember : Params.LeavingMembers)
 		{
-			const FOnlineAccountIdHandle LeavingMemberAccountId = FindAccountId(LeavingMember.Key);
+			const FAccountId LeavingMemberAccountId = FindAccountId(LeavingMember.Key);
 			if (LeavingMemberAccountId.IsValid())
 			{
 				LeavingMemberReason.Add(LeavingMemberAccountId, LeavingMember.Value);
@@ -1821,7 +1821,7 @@ TOnlineAsyncOpHandle<FLobbiesProcessLobbyNotificationImpl> FLobbiesEOSGS::Proces
 			&LobbyEvents);
 
 		// Remove active users if needed.
-		for (FOnlineAccountIdHandle LeavingMember : Result.LeavingLocalMembers)
+		for (FAccountId LeavingMember : Result.LeavingLocalMembers)
 		{
 			RemoveActiveLobby(LeavingMember, Params.LobbyData.ToSharedRef());
 		}
diff --git a/Engine/Plugins/Online/Experimental/OnlineServicesEOSGS/Source/Private/Online/LobbiesEOSGSTypes.cpp b/Engine/Plugins/Online/Experimental/OnlineServicesEOSGS/Source/Private/Online/LobbiesEOSGSTypes.cpp
--- a/Engine/Plugins/Online/Experimental/OnlineServicesEOSGS/Source/Private/Online/LobbiesEOSGSTypes.cpp
+++ b/Engine/Plugins/Online/Experimental/OnlineServicesEOSGS/Source/Private/Online/LobbiesEOSGSTypes.cpp
@@ -147,7 +147,7 @@ const EOS_HLobbyDetails FLobbyDetailsEOS::InvalidLobbyDetailsHandle = {};
 
 TDefaultErrorResultInternal<TSharedRef<FLobbyDetailsEOS>> FLobbyDetailsEOS::CreateFromLobbyId(
 	const TSharedRef<FLobbyPrerequisitesEOS>& Prerequisites,
-	FOnlineAccountIdHandle LocalUserId,
+	FAccountId LocalUserId,
 	EOS_LobbyId LobbyId)
 {
 	EOS_HLobbyDetails LobbyDetailsHandle = {};
@@ -176,7 +176,7 @@ TDefaultErrorResultInternal<TSharedRef<FLobbyDetailsEOS>> FLobbyDetailsEOS::Crea
 
 TDefaultErrorResultInternal<TSharedRef<FLobbyDetailsEOS>> FLobbyDetailsEOS::CreateFromInviteId(
 	const TSharedRef<FLobbyPrerequisitesEOS>& Prerequisites,
-	FOnlineAccountIdHandle LocalUserId,
+	FAccountId LocalUserId,
 	const char* InviteId)
 {
 	EOS_HLobbyDetails LobbyDetailsHandle = {};
@@ -204,7 +204,7 @@ TDefaultErrorResultInternal<TSharedRef<FLobbyDetailsEOS>> FLobbyDetailsEOS::Crea
 
 TDefaultErrorResultInternal<TSharedRef<FLobbyDetailsEOS>> FLobbyDetailsEOS::CreateFromUiEventId(
 	const TSharedRef<FLobbyPrerequisitesEOS>& Prerequisites,
-	FOnlineAccountIdHandle LocalUserId,
+	FAccountId LocalUserId,
 	EOS_UI_EventId UiEventId)
 {
 	EOS_HLobbyDetails LobbyDetailsHandle = {};
@@ -232,7 +232,7 @@ TDefaultErrorResultInternal<TSharedRef<FLobbyDetailsEOS>> FLobbyDetailsEOS::Crea
 
 TDefaultErrorResultInternal<TSharedRef<FLobbyDetailsEOS>> FLobbyDetailsEOS::CreateFromSearchResult(
 	const TSharedRef<FLobbyPrerequisitesEOS>& Prerequisites,
-	FOnlineAccountIdHandle LocalUserId,
+	FAccountId LocalUserId,
 	EOS_HLobbySearch SearchHandle,
 	uint32_t ResultIndex)
 {
@@ -294,11 +294,11 @@ TFuture<TDefaultErrorResultInternal<TSharedRef<FClientLobbySnapshot>>> FLobbyDet
 	TPromise<TDefaultErrorResultInternal<TSharedRef<FClientLobbySnapshot>>> Promise;
 	TFuture<TDefaultErrorResultInternal<TSharedRef<FClientLobbySnapshot>>> Future = Promise.GetFuture();
 
-	// Resolve lobby member product user ids to FOnlineAccountIdHandle before proceeding.
+	// Resolve lobby member product user ids to FAccountId before proceeding.
 	AuthInterface->ResolveAccountIds(AssociatedLocalUser, *MemberProductUserIds)
-	.Then([StrongThis = AsShared(), Promise = MoveTemp(Promise), MemberProductUserIds](TFuture<TArray<FOnlineAccountIdHandle>>&& Future) mutable
+	.Then([StrongThis = AsShared(), Promise = MoveTemp(Promise), MemberProductUserIds](TFuture<TArray<FAccountId>>&& Future) mutable
 	{
-		const TArray<FOnlineAccountIdHandle>& ResolvedAccountIds = Future.Get();
+		const TArray<FAccountId>& ResolvedAccountIds = Future.Get();
 		if (MemberProductUserIds->Num() != ResolvedAccountIds.Num())
 		{
 			// Todo: Errors
@@ -321,7 +321,7 @@ TFuture<TDefaultErrorResultInternal<TSharedRef<FClientLobbySnapshot>>> FLobbyDet
 			for (int32 MemberIndex = 0; MemberIndex < MemberProductUserIds->Num(); ++MemberIndex)
 			{
 				const EOS_ProductUserId MemberProductUserId = (*MemberProductUserIds)[MemberIndex];
-				const FOnlineAccountIdHandle ResolvedMemberAccountId = ResolvedAccountIds[MemberIndex];
+				const FAccountId ResolvedMemberAccountId = ResolvedAccountIds[MemberIndex];
 
 				if (MemberProductUserId == LobbyOwner)
 				{
@@ -371,7 +371,7 @@ TFuture<TDefaultErrorResultInternal<TSharedRef<FClientLobbySnapshot>>> FLobbyDet
 	return Future;
 }
 
-TDefaultErrorResultInternal<TSharedRef<FClientLobbyMemberSnapshot>> FLobbyDetailsEOS::GetLobbyMemberSnapshot(FOnlineAccountIdHandle MemberAccountId) const
+TDefaultErrorResultInternal<TSharedRef<FClientLobbyMemberSnapshot>> FLobbyDetailsEOS::GetLobbyMemberSnapshot(FAccountId MemberAccountId) const
 {
 	EOS_ProductUserId MemberProductUserId = GetProductUserIdChecked(MemberAccountId);
 
@@ -418,7 +418,7 @@ TDefaultErrorResultInternal<TSharedRef<FClientLobbyMemberSnapshot>> FLobbyDetail
 }
 
 TFuture<EOS_EResult> FLobbyDetailsEOS::ApplyLobbyDataUpdateFromLocalChanges(
-	FOnlineAccountIdHandle LocalUserId,
+	FAccountId LocalUserId,
 	const FClientLobbyDataChanges& Changes) const
 {
 	EOS_HLobbyModification LobbyModificationHandle = nullptr;
@@ -513,7 +513,7 @@ TFuture<EOS_EResult> FLobbyDetailsEOS::ApplyLobbyDataUpdateFromLocalChanges(
 }
 
 TFuture<EOS_EResult> FLobbyDetailsEOS::ApplyLobbyMemberDataUpdateFromLocalChanges(
-	FOnlineAccountIdHandle LocalUserId,
+	FAccountId LocalUserId,
 	const FClientLobbyMemberDataChanges& Changes) const
 {
 	EOS_HLobbyModification LobbyModificationHandle = {};
@@ -595,7 +595,7 @@ TFuture<EOS_EResult> FLobbyDetailsEOS::ApplyLobbyMemberDataUpdateFromLocalChange
 FLobbyDetailsEOS::FLobbyDetailsEOS(
 	const TSharedRef<FLobbyPrerequisitesEOS>& Prerequisites,
 	const TSharedRef<FLobbyDetailsInfoEOS>& LobbyDetailsInfo,
-	FOnlineAccountIdHandle LocalUserId,
+	FAccountId LocalUserId,
 	ELobbyDetailsSource LobbyDetailsSource,
 	EOS_HLobbyDetails LobbyDetailsHandle)
 	: Prerequisites(Prerequisites)
@@ -642,7 +642,7 @@ FLobbyDataEOS::~FLobbyDataEOS()
 	}
 }
 
-void FLobbyDataEOS::AddUserLobbyDetails(FOnlineAccountIdHandle LocalUserId, const TSharedPtr<FLobbyDetailsEOS>& LobbyDetails)
+void FLobbyDataEOS::AddUserLobbyDetails(FAccountId LocalUserId, const TSharedPtr<FLobbyDetailsEOS>& LobbyDetails)
 {
 	if (TSharedPtr<FLobbyDetailsEOS> ExistingDetails = GetUserLobbyDetails(LocalUserId))
 	{
@@ -655,7 +655,7 @@ void FLobbyDataEOS::AddUserLobbyDetails(FOnlineAccountIdHandle LocalUserId, cons
 	UserLobbyDetails.Add(LocalUserId, LobbyDetails);
 }
 
-TSharedPtr<FLobbyDetailsEOS> FLobbyDataEOS::GetUserLobbyDetails(FOnlineAccountIdHandle LocalUserId) const
+TSharedPtr<FLobbyDetailsEOS> FLobbyDataEOS::GetUserLobbyDetails(FAccountId LocalUserId) const
 {
 	const TSharedPtr<FLobbyDetailsEOS>* Result = UserLobbyDetails.Find(LocalUserId);
 	return Result ? *Result : TSharedPtr<FLobbyDetailsEOS>();
@@ -665,7 +665,7 @@ TSharedPtr<FLobbyDetailsEOS> FLobbyDataEOS::GetActiveLobbyDetails() const
 {
 	TSharedPtr<FLobbyDetailsEOS> FoundDetails;
 
-	for (const TPair<FOnlineAccountIdHandle, TSharedPtr<FLobbyDetailsEOS>>& LobbyDetails : UserLobbyDetails)
+	for (const TPair<FAccountId, TSharedPtr<FLobbyDetailsEOS>>& LobbyDetails : UserLobbyDetails)
 	{
 		if (LobbyDetails.Value->GetDetailsSource() == ELobbyDetailsSource::Active)
 		{
@@ -717,8 +717,8 @@ TFuture<TDefaultErrorResultInternal<TSharedRef<FLobbyDataEOS>>> FLobbyDataEOS::C
 		TSharedRef<FClientLobbyData> LobbyData = MakeShared<FClientLobbyData>(LobbyIdHandle);
 
 		// Fetch member data and apply them to the lobby.
-		TMap<FOnlineAccountIdHandle, TSharedRef<FClientLobbyMemberSnapshot>> MemberSnapshots;
-		for (FOnlineAccountIdHandle MemberAccountId : LobbySnapshot->Members)
+		TMap<FAccountId, TSharedRef<FClientLobbyMemberSnapshot>> MemberSnapshots;
+		for (FAccountId MemberAccountId : LobbySnapshot->Members)
 		{
 			TDefaultErrorResultInternal<TSharedRef<FClientLobbyMemberSnapshot>> LobbyMemberSnapshotResult = LobbyDetails->GetLobbyMemberSnapshot(MemberAccountId);
 			if (LobbyMemberSnapshotResult.IsError())
@@ -759,7 +759,7 @@ TSharedPtr<FLobbyDataEOS> FLobbyDataRegistryEOS::Find(FOnlineLobbyIdHandle Lobby
 	return Result ? Result->Pin() : TSharedPtr<FLobbyDataEOS>();
 }
 
-TFuture<TDefaultErrorResultInternal<TSharedRef<FLobbyDataEOS>>> FLobbyDataRegistryEOS::FindOrCreateFromLobbyDetails(FOnlineAccountIdHandle LocalUserId, const TSharedRef<FLobbyDetailsEOS>& LobbyDetails)
+TFuture<TDefaultErrorResultInternal<TSharedRef<FLobbyDataEOS>>> FLobbyDataRegistryEOS::FindOrCreateFromLobbyDetails(FAccountId LocalUserId, const TSharedRef<FLobbyDetailsEOS>& LobbyDetails)
 {
 	if (TSharedPtr<FLobbyDataEOS> FindResult = Find(LobbyDetails->GetInfo()->GetLobbyId()))
 	{
@@ -818,7 +818,7 @@ FLobbyDataEOS::FUnregisterFn FLobbyDataRegistryEOS::MakeUnregisterFn()
 TFuture<TDefaultErrorResultInternal<TSharedRef<FLobbyInviteDataEOS>>> FLobbyInviteDataEOS::CreateFromInviteId(
 	const TSharedRef<FLobbyPrerequisitesEOS>& Prerequisites,
 	const TSharedRef<FLobbyDataRegistryEOS>& LobbyDataRegistry,
-	FOnlineAccountIdHandle LocalUserId,
+	FAccountId LocalUserId,
 	const char* InviteIdEOS,
 	EOS_ProductUserId Sender)
 {
@@ -843,8 +843,8 @@ TFuture<TDefaultErrorResultInternal<TSharedRef<FLobbyInviteDataEOS>>> FLobbyInvi
 			return;
 		}
 
-		// Once the lobby data has been resolved the FOnlineAccountIdHandle for the sender is expected to be in the AccountID cache.
-		const FOnlineAccountIdHandle SenderUserId = FindAccountId(Sender);
+		// Once the lobby data has been resolved the FAccountId for the sender is expected to be in the AccountID cache.
+		const FAccountId SenderUserId = FindAccountId(Sender);
 		if (!SenderUserId.IsValid())
 		{
 			// Todo: Errors.
@@ -860,8 +860,8 @@ TFuture<TDefaultErrorResultInternal<TSharedRef<FLobbyInviteDataEOS>>> FLobbyInvi
 
 FLobbyInviteDataEOS::FLobbyInviteDataEOS(
 	const TSharedRef<FLobbyInviteIdEOS>& InviteIdEOS,
-	FOnlineAccountIdHandle Receiver,
-	FOnlineAccountIdHandle Sender,
+	FAccountId Receiver,
+	FAccountId Sender,
 	const TSharedRef<FLobbyDetailsEOS>& LobbyDetails,
 	const TSharedRef<FLobbyDataEOS>& LobbyData)
 	: InviteIdEOS(InviteIdEOS)
diff --git a/Engine/Plugins/Online/Experimental/OnlineServicesEOSGS/Source/Private/Online/LobbiesEOSGSTypes.h b/Engine/Plugins/Online/Experimental/OnlineServicesEOSGS/Source/Private/Online/LobbiesEOSGSTypes.h
--- a/Engine/Plugins/Online/Experimental/OnlineServicesEOSGS/Source/Private/Online/LobbiesEOSGSTypes.h
+++ b/Engine/Plugins/Online/Experimental/OnlineServicesEOSGS/Source/Private/Online/LobbiesEOSGSTypes.h
@@ -213,19 +213,19 @@ public:
 
 	static TDefaultErrorResultInternal<TSharedRef<FLobbyDetailsEOS>> CreateFromLobbyId(
 		const TSharedRef<FLobbyPrerequisitesEOS>& Prerequisites,
-		FOnlineAccountIdHandle LocalUserId,
+		FAccountId LocalUserId,
 		EOS_LobbyId LobbyId);
 	static TDefaultErrorResultInternal<TSharedRef<FLobbyDetailsEOS>> CreateFromInviteId(
 		const TSharedRef<FLobbyPrerequisitesEOS>& Prerequisites,
-		FOnlineAccountIdHandle LocalUserId,
+		FAccountId LocalUserId,
 		const char* InviteId);
 	static TDefaultErrorResultInternal<TSharedRef<FLobbyDetailsEOS>> CreateFromUiEventId(
 		const TSharedRef<FLobbyPrerequisitesEOS>& Prerequisites,
-		FOnlineAccountIdHandle LocalUserId,
+		FAccountId LocalUserId,
 		EOS_UI_EventId UiEventId);
 	static TDefaultErrorResultInternal<TSharedRef<FLobbyDetailsEOS>> CreateFromSearchResult(
 		const TSharedRef<FLobbyPrerequisitesEOS>& Prerequisites,
-		FOnlineAccountIdHandle LocalUserId,
+		FAccountId LocalUserId,
 		EOS_HLobbySearch SearchHandle,
 		uint32_t ResultIndex);
 
@@ -235,7 +235,7 @@ public:
 	EOS_HLobbyDetails GetEOSHandle() const { return LobbyDetailsHandle; }
 	const TSharedRef<FLobbyDetailsInfoEOS>& GetInfo() const { return LobbyDetailsInfo; }
 	ELobbyDetailsSource GetDetailsSource() const { return LobbyDetailsSource; }
-	FOnlineAccountIdHandle GetAssociatedUser() const { return AssociatedLocalUser; }
+	FAccountId GetAssociatedUser() const { return AssociatedLocalUser; }
 
 	/**
 	 * Retrieve a lobby data snapshot from the EOS lobby details object.
@@ -247,19 +247,19 @@ public:
 	 * Retrieve lobby member data snapshot from the EOS lobby details object.
 	 * The lobby schema will be used to translate attribute data before returning.
 	 */
-	TDefaultErrorResultInternal<TSharedRef<FClientLobbyMemberSnapshot>> GetLobbyMemberSnapshot(FOnlineAccountIdHandle MemberHandle) const;
+	TDefaultErrorResultInternal<TSharedRef<FClientLobbyMemberSnapshot>> GetLobbyMemberSnapshot(FAccountId MemberHandle) const;
 
 	/**
 	 * Apply client side lobby changes to the lobby service.
 	 * The lobby schema will be used to translate any changed attributes before sending.
 	 */
-	TFuture<EOS_EResult> ApplyLobbyDataUpdateFromLocalChanges(FOnlineAccountIdHandle LocalUserId, const FClientLobbyDataChanges& Changes) const;
+	TFuture<EOS_EResult> ApplyLobbyDataUpdateFromLocalChanges(FAccountId LocalUserId, const FClientLobbyDataChanges& Changes) const;
 
 	/**
 	 * Apply client side lobby member changes to the lobby service.
 	 * The lobby schema will be used to translate any changed attributes before sending.
 	 */
-	TFuture<EOS_EResult> ApplyLobbyMemberDataUpdateFromLocalChanges(FOnlineAccountIdHandle LocalUserId, const FClientLobbyMemberDataChanges& Changes) const;
+	TFuture<EOS_EResult> ApplyLobbyMemberDataUpdateFromLocalChanges(FAccountId LocalUserId, const FClientLobbyMemberDataChanges& Changes) const;
 
 private:
 	template <typename, ESPMode>
@@ -267,14 +267,14 @@ private:
 	FLobbyDetailsEOS(
 		const TSharedRef<FLobbyPrerequisitesEOS>& Prerequisites,
 		const TSharedRef<FLobbyDetailsInfoEOS>& LobbyDetailsInfo,
-		FOnlineAccountIdHandle LocalUserId,
+		FAccountId LocalUserId,
 		ELobbyDetailsSource LobbyDetailsSource,
 		EOS_HLobbyDetails LobbyDetailsHandle);
 
 	static const EOS_HLobbyDetails InvalidLobbyDetailsHandle;
 	TSharedRef<FLobbyPrerequisitesEOS> Prerequisites;
 	TSharedRef<FLobbyDetailsInfoEOS> LobbyDetailsInfo;
-	FOnlineAccountIdHandle AssociatedLocalUser;
+	FAccountId AssociatedLocalUser;
 	ELobbyDetailsSource LobbyDetailsSource;
 	EOS_HLobbyDetails LobbyDetailsHandle = InvalidLobbyDetailsHandle;
 };
@@ -334,8 +334,8 @@ public:
 	EOS_LobbyId GetLobbyIdEOS() const { return LobbyDetailsInfo->GetLobbyId(); }
 	const FString& GetLobbyId() const { return LobbyId; }
 
-	void AddUserLobbyDetails(FOnlineAccountIdHandle LocalUserId, const TSharedPtr<FLobbyDetailsEOS>& LobbyDetails);
-	TSharedPtr<FLobbyDetailsEOS> GetUserLobbyDetails(FOnlineAccountIdHandle LocalUserId) const;
+	void AddUserLobbyDetails(FAccountId LocalUserId, const TSharedPtr<FLobbyDetailsEOS>& LobbyDetails);
+	TSharedPtr<FLobbyDetailsEOS> GetUserLobbyDetails(FAccountId LocalUserId) const;
 
 	/**
 	 * Active lobby details are needed to process lobby notifications. Search for and return active
@@ -356,7 +356,7 @@ private:
 	TSharedRef<FLobbyDetailsInfoEOS> LobbyDetailsInfo;
 	FUnregisterFn UnregisterFn;
 	FString LobbyId;
-	TMap<FOnlineAccountIdHandle, TSharedPtr<FLobbyDetailsEOS>> UserLobbyDetails;
+	TMap<FAccountId, TSharedPtr<FLobbyDetailsEOS>> UserLobbyDetails;
 };
 
 // Todo: implement handling for FOnlineIdRegistryRegistry.
@@ -367,7 +367,7 @@ public:
 
 	TSharedPtr<FLobbyDataEOS> Find(EOS_LobbyId EOSLobbyId) const;
 	TSharedPtr<FLobbyDataEOS> Find(FOnlineLobbyIdHandle LobbyIdHandle) const;
-	TFuture<TDefaultErrorResultInternal<TSharedRef<FLobbyDataEOS>>> FindOrCreateFromLobbyDetails(FOnlineAccountIdHandle LocalUserId, const TSharedRef<FLobbyDetailsEOS>& LobbyDetails);
+	TFuture<TDefaultErrorResultInternal<TSharedRef<FLobbyDataEOS>>> FindOrCreateFromLobbyDetails(FAccountId LocalUserId, const TSharedRef<FLobbyDetailsEOS>& LobbyDetails);
 
 private:
 	void Register(const TSharedRef<FLobbyDataEOS>& LobbyIdHandleData);
@@ -407,7 +407,7 @@ public:
 	static TFuture<TDefaultErrorResultInternal<TSharedRef<FLobbyInviteDataEOS>>> CreateFromInviteId(
 		const TSharedRef<FLobbyPrerequisitesEOS>& Prerequisites,
 		const TSharedRef<FLobbyDataRegistryEOS>& LobbyDataRegistry,
-		FOnlineAccountIdHandle LocalUserId,
+		FAccountId LocalUserId,
 		const char* InviteIdEOS,
 		EOS_ProductUserId Sender);
 
@@ -417,8 +417,8 @@ public:
 	const char* GetInviteIdEOS() const { return InviteIdEOS->Get(); }
 	const FString& GetInviteId() const { return InviteId; }
 
-	FOnlineAccountIdHandle GetReceiver() const { return Receiver; }
-	FOnlineAccountIdHandle GetSender() const { return Sender; }
+	FAccountId GetReceiver() const { return Receiver; }
+	FAccountId GetSender() const { return Sender; }
 
 private:
 	template <typename, ESPMode>
@@ -426,14 +426,14 @@ private:
 
 	FLobbyInviteDataEOS(
 		const TSharedRef<FLobbyInviteIdEOS>& InviteIdEOS,
-		FOnlineAccountIdHandle Receiver,
-		FOnlineAccountIdHandle Sender,
+		FAccountId Receiver,
+		FAccountId Sender,
 		const TSharedRef<FLobbyDetailsEOS>& LobbyDetails,
 		const TSharedRef<FLobbyDataEOS>& LobbyData);
 
 	TSharedRef<FLobbyInviteIdEOS> InviteIdEOS;
-	FOnlineAccountIdHandle Receiver;
-	FOnlineAccountIdHandle Sender;
+	FAccountId Receiver;
+	FAccountId Sender;
 	TSharedRef<FLobbyDetailsEOS> LobbyDetails;
 	TSharedRef<FLobbyDataEOS> LobbyData;
 	FString InviteId;
diff --git a/Engine/Plugins/Online/Experimental/OnlineServicesEOSGS/Source/Private/Online/OnlineIdEOSGS.cpp b/Engine/Plugins/Online/Experimental/OnlineServicesEOSGS/Source/Private/Online/OnlineIdEOSGS.cpp
--- a/Engine/Plugins/Online/Experimental/OnlineServicesEOSGS/Source/Private/Online/OnlineIdEOSGS.cpp
+++ b/Engine/Plugins/Online/Experimental/OnlineServicesEOSGS/Source/Private/Online/OnlineIdEOSGS.cpp
@@ -18,17 +18,17 @@ FOnlineAccountIdRegistryEOSGS::FOnlineAccountIdRegistryEOSGS()
 
 }
 
-FOnlineAccountIdHandle FOnlineAccountIdRegistryEOSGS::FindAccountId(const EOS_ProductUserId ProductUserId) const
+FAccountId FOnlineAccountIdRegistryEOSGS::FindAccountId(const EOS_ProductUserId ProductUserId) const
 {
 	return Registry.FindHandle(ProductUserId);
 }
 
-EOS_ProductUserId FOnlineAccountIdRegistryEOSGS::GetProductUserId(const FOnlineAccountIdHandle& Handle) const
+EOS_ProductUserId FOnlineAccountIdRegistryEOSGS::GetProductUserId(const FAccountId& Handle) const
 {
 	return Registry.FindIdValue(Handle);
 }
 
-FString FOnlineAccountIdRegistryEOSGS::ToLogString(const FOnlineAccountIdHandle& Handle) const
+FString FOnlineAccountIdRegistryEOSGS::ToLogString(const FAccountId& Handle) const
 {
 	if (Registry.ValidateOnlineId(Handle))
 	{
@@ -38,7 +38,7 @@ FString FOnlineAccountIdRegistryEOSGS::ToLogString(const FOnlineAccountIdHandle&
 	return FString();
 }
 
-TArray<uint8> FOnlineAccountIdRegistryEOSGS::ToReplicationData(const FOnlineAccountIdHandle& Handle) const
+TArray<uint8> FOnlineAccountIdRegistryEOSGS::ToReplicationData(const FAccountId& Handle) const
 {
 	TArray<uint8> ReplicationData;
 	if (Registry.ValidateOnlineId(Handle))
@@ -59,7 +59,7 @@ TArray<uint8> FOnlineAccountIdRegistryEOSGS::ToReplicationData(const FOnlineAcco
 	}
 	return ReplicationData;
 }
-FOnlineAccountIdHandle FOnlineAccountIdRegistryEOSGS::FromReplicationData(const TArray<uint8>& ReplicationData)
+FAccountId FOnlineAccountIdRegistryEOSGS::FromReplicationData(const TArray<uint8>& ReplicationData)
 {
 	if (ReplicationData.Num() == OnlineIdEOSHexBufferLength)
 	{
@@ -84,7 +84,7 @@ FOnlineAccountIdRegistryEOSGS& FOnlineAccountIdRegistryEOSGS::Get()
 	return Instance;
 }
 
-FOnlineAccountIdHandle FOnlineAccountIdRegistryEOSGS::FindOrAddAccountId(const EOS_ProductUserId ProductUserId)
+FAccountId FOnlineAccountIdRegistryEOSGS::FindOrAddAccountId(const EOS_ProductUserId ProductUserId)
 {
 	if (ensure(EOS_ProductUserId_IsValid(ProductUserId)))
 	{
@@ -93,26 +93,26 @@ FOnlineAccountIdHandle FOnlineAccountIdRegistryEOSGS::FindOrAddAccountId(const E
 	return Registry.GetInvalidHandle();
 }
 
-EOS_ProductUserId GetProductUserId(const FOnlineAccountIdHandle& Handle)
+EOS_ProductUserId GetProductUserId(const FAccountId& Handle)
 {
 	return FOnlineAccountIdRegistryEOSGS::GetRegistered().GetProductUserId(Handle);
 }
 
-EOS_ProductUserId GetProductUserIdChecked(const FOnlineAccountIdHandle& Handle)
+EOS_ProductUserId GetProductUserIdChecked(const FAccountId& Handle)
 {
 	EOS_ProductUserId ProductUserId = GetProductUserId(Handle);
 	check(EOS_ProductUserId_IsValid(ProductUserId) == EOS_TRUE);
 	return ProductUserId;
 }
 
-FOnlineAccountIdHandle FindAccountId(const EOS_ProductUserId ProductUserId)
+FAccountId FindAccountId(const EOS_ProductUserId ProductUserId)
 {
 	return FOnlineAccountIdRegistryEOSGS::GetRegistered().FindAccountId(ProductUserId);
 }
 
-FOnlineAccountIdHandle FindAccountIdChecked(const EOS_ProductUserId ProductUserId)
+FAccountId FindAccountIdChecked(const EOS_ProductUserId ProductUserId)
 {
-	FOnlineAccountIdHandle Result = FindAccountId(ProductUserId);
+	FAccountId Result = FindAccountId(ProductUserId);
 	check(Result.IsValid());
 	return Result;
 }
diff --git a/Engine/Plugins/Online/Experimental/OnlineServicesEOSGS/Source/Private/Online/SessionsEOSGS.cpp b/Engine/Plugins/Online/Experimental/OnlineServicesEOSGS/Source/Private/Online/SessionsEOSGS.cpp
--- a/Engine/Plugins/Online/Experimental/OnlineServicesEOSGS/Source/Private/Online/SessionsEOSGS.cpp
+++ b/Engine/Plugins/Online/Experimental/OnlineServicesEOSGS/Source/Private/Online/SessionsEOSGS.cpp
@@ -97,7 +97,7 @@ FSessionEOSGS::FSessionEOSGS(const EOS_HSessionDetails& InSessionDetailsHandle)
 						// We retrieve the registered player id
 						const FString& PlayerIdStr = KeyComponents[1];
 						const EOS_ProductUserId ProductUserId = EOS_ProductUserId_FromString(TCHAR_TO_UTF8(*PlayerIdStr));
-						const FOnlineAccountIdHandle& IdHandle = FindAccountId(ProductUserId);
+						const FAccountId& IdHandle = FindAccountId(ProductUserId);
 
 						// This constructor should only be used by BuildSessionFromDetailsHandle, after all user ids in the session have been resolved
 						check(IdHandle.IsValid());
@@ -125,7 +125,7 @@ FSessionEOSGS::FSessionEOSGS(const EOS_HSessionDetails& InSessionDetailsHandle)
 						// We retrieve the member id
 						FString PlayerIdStr = KeyComponents[0];
 						EOS_ProductUserId ProductUserId = EOS_ProductUserId_FromString(TCHAR_TO_UTF8(*PlayerIdStr));
-						FOnlineAccountIdHandle IdHandle = FindAccountId(ProductUserId);
+						FAccountId IdHandle = FindAccountId(ProductUserId);
 
 						FSessionMember& SessionMember = SessionSettings.SessionMembers.FindOrAdd(IdHandle);
 
@@ -276,17 +276,17 @@ void FSessionsEOSGS::UnregisterEventHandlers()
 
 void FSessionsEOSGS::HandleSessionInviteReceived(const EOS_Sessions_SessionInviteReceivedCallbackInfo* Data)
 {
-	FOnlineAccountIdHandle LocalUserId = FindAccountId(Data->LocalUserId);
+	FAccountId LocalUserId = FindAccountId(Data->LocalUserId);
 	if (LocalUserId.IsValid())
 	{
 		Services.Get<FAuthEOSGS>()->ResolveAccountIds(LocalUserId, { Data->LocalUserId, Data->TargetUserId })
-			.Next([this, WeakThis = AsWeak(), LocalUserId, InviteId = FString(UTF8_TO_TCHAR(Data->InviteId))](TArray<FOnlineAccountIdHandle> ResolvedAccountIds)
+			.Next([this, WeakThis = AsWeak(), LocalUserId, InviteId = FString(UTF8_TO_TCHAR(Data->InviteId))](TArray<FAccountId> ResolvedAccountIds)
 		{
 			if (const TSharedPtr<ISessions> StrongThis = WeakThis.Pin())
 			{
 				// First and second place in the array will be occupied by the receiver and the sender, respectively, since the same order is kept in the array of resolved ids
-				const FOnlineAccountIdHandle& ReceiverId = ResolvedAccountIds[0];
-				const FOnlineAccountIdHandle& SenderId = ResolvedAccountIds[1];
+				const FAccountId& ReceiverId = ResolvedAccountIds[0];
+				const FAccountId& SenderId = ResolvedAccountIds[1];
 
 				TResult<TSharedRef<const FSession>, FOnlineError> SessionInviteResult = BuildSessionFromInvite(InviteId);
 				
@@ -317,7 +317,7 @@ void FSessionsEOSGS::HandleSessionInviteReceived(const EOS_Sessions_SessionInvit
 
 void FSessionsEOSGS::HandleSessionInviteAccepted(const EOS_Sessions_SessionInviteAcceptedCallbackInfo* Data)
 {
-	FOnlineAccountIdHandle LocalUserId = FindAccountId(Data->LocalUserId);
+	FAccountId LocalUserId = FindAccountId(Data->LocalUserId);
 	if (LocalUserId.IsValid())
 	{
 		FUISessionJoinRequested Event{
@@ -334,7 +334,7 @@ void FSessionsEOSGS::HandleSessionInviteAccepted(const EOS_Sessions_SessionInvit
 
 void FSessionsEOSGS::HandleJoinSessionAccepted(const EOS_Sessions_JoinSessionAcceptedCallbackInfo* Data)
 {
-	FOnlineAccountIdHandle LocalUserId = FindAccountId(Data->LocalUserId);
+	FAccountId LocalUserId = FindAccountId(Data->LocalUserId);
 	if (LocalUserId.IsValid())
 	{
 		FUISessionJoinRequested Event {
@@ -639,7 +639,7 @@ void FSessionsEOSGS::WriteCreateSessionModificationHandle(EOS_HSessionModificati
 
 	// Session Members
 
-	for (const TPair<FOnlineAccountIdHandle, FSessionMember>& SessionMemberEntry : SessionSettings.SessionMembers)
+	for (const TPair<FAccountId, FSessionMember>& SessionMemberEntry : SessionSettings.SessionMembers)
 	{
 		const FSessionMember& SessionMember = SessionMemberEntry.Value;
 		
@@ -652,7 +652,7 @@ void FSessionsEOSGS::WriteCreateSessionModificationHandle(EOS_HSessionModificati
 
 	// TODO: Should always be empty at session creation time so might not be needed
 	// Registered Players
-	for (const TPair<FOnlineAccountIdHandle, FRegisteredPlayer>& Entry : SessionSettings.RegisteredPlayers)
+	for (const TPair<FAccountId, FRegisteredPlayer>& Entry : SessionSettings.RegisteredPlayers)
 	{
 		const FName& HasReservedSlotKey = FName(EOSGS_REGISTERED_PLAYERS.ToString() + TEXT(":") + LexToString(GetProductUserIdChecked(Entry.Key)) + TEXT(":") + EOSGS_REGISTERED_PLAYER_HAS_RESERVED_SLOT.ToString());
 		AddAttribute(SessionModificationHandle, HasReservedSlotKey, { FSchemaVariant(Entry.Value.bHasReservedSlot), ESchemaAttributeVisibility::Public });
@@ -744,7 +744,7 @@ void FSessionsEOSGS::WriteUpdateSessionModificationHandle(EOS_HSessionModificati
 
 	// Session Members
 	
-	for (const FOnlineAccountIdHandle& SessionMemberId : NewSettings.RemovedSessionMembers)
+	for (const FAccountId& SessionMemberId : NewSettings.RemovedSessionMembers)
 	{
 		// To remove a session member, we'll remove all their attributes
  		if (const FSessionMember* SessionMember = SessionSettings.SessionMembers.Find(SessionMemberId))
@@ -759,7 +759,7 @@ void FSessionsEOSGS::WriteUpdateSessionModificationHandle(EOS_HSessionModificati
  		}
 	}
 
-	for (const TPair<FOnlineAccountIdHandle, FSessionMemberUpdate>& SessionMemberEntry : NewSettings.UpdatedSessionMembers)
+	for (const TPair<FAccountId, FSessionMemberUpdate>& SessionMemberEntry : NewSettings.UpdatedSessionMembers)
 	{
 		const FSessionMemberUpdate& SessionMemberUpdate = SessionMemberEntry.Value;
 
@@ -778,7 +778,7 @@ void FSessionsEOSGS::WriteUpdateSessionModificationHandle(EOS_HSessionModificati
 
 	// Registered Players
 	
-	for (const FOnlineAccountIdHandle& RemovedEntry : NewSettings.RemovedRegisteredPlayers)
+	for (const FAccountId& RemovedEntry : NewSettings.RemovedRegisteredPlayers)
 	{
 		const FSchemaAttributeId& HasReservedSlotKey = FSchemaAttributeId(EOSGS_REGISTERED_PLAYERS.ToString() + TEXT(":") + LexToString(GetProductUserIdChecked(RemovedEntry)) + TEXT(":") + EOSGS_REGISTERED_PLAYER_HAS_RESERVED_SLOT.ToString());
 		RemoveAttribute(SessionModificationHandle, HasReservedSlotKey);
@@ -787,7 +787,7 @@ void FSessionsEOSGS::WriteUpdateSessionModificationHandle(EOS_HSessionModificati
 		RemoveAttribute(SessionModificationHandle, IsInSessionKey);
 	}
 
-	for (const TPair<FOnlineAccountIdHandle, FRegisteredPlayer>& Entry : NewSettings.UpdatedRegisteredPlayers)
+	for (const TPair<FAccountId, FRegisteredPlayer>& Entry : NewSettings.UpdatedRegisteredPlayers)
 	{
 		const FSchemaAttributeId& HasReservedSlotKey = FSchemaAttributeId(EOSGS_REGISTERED_PLAYERS.ToString() + TEXT(":") + LexToString(GetProductUserIdChecked(Entry.Key)) + TEXT(":") + EOSGS_REGISTERED_PLAYER_HAS_RESERVED_SLOT.ToString());
 		AddAttribute(SessionModificationHandle, HasReservedSlotKey, { FSchemaVariant(Entry.Value.bHasReservedSlot), ESchemaAttributeVisibility::Public });
@@ -1183,7 +1183,7 @@ void FSessionsEOSGS::SetSessionSearchSessionId(FSessionSearchHandleEOSGS& Sessio
 	}
 }
 
-void FSessionsEOSGS::SetSessionSearchTargetId(FSessionSearchHandleEOSGS& SessionSearchHandle, const FOnlineAccountIdHandle& TargetUserId)
+void FSessionsEOSGS::SetSessionSearchTargetId(FSessionSearchHandleEOSGS& SessionSearchHandle, const FAccountId& TargetUserId)
 {
 	EOS_SessionSearch_SetTargetUserIdOptions Options = { };
 	Options.ApiVersion = EOS_SESSIONSEARCH_SETTARGETUSERID_API_LATEST;
@@ -1646,7 +1646,7 @@ TOnlineAsyncOpHandle<FBuildSessionFromDetailsHandle> FSessionsEOSGS::BuildSessio
 	if (GetIdsResult.IsOk())
 	{
 		Services.Get<FAuthEOSGS>()->ResolveAccountIds(Params.LocalUserId, GetIdsResult.GetOkValue())
-			.Next([this, WeakOp = Op->AsWeak(), Params = MoveTemp(Params)](TArray<FOnlineAccountIdHandle> ResolvedAccountIds) mutable
+			.Next([this, WeakOp = Op->AsWeak(), Params = MoveTemp(Params)](TArray<FAccountId> ResolvedAccountIds) mutable
 		{
 			if (TOnlineAsyncOpPtr<FBuildSessionFromDetailsHandle> StrongOp = WeakOp.Pin())
 			{
@@ -1680,7 +1680,7 @@ TOnlineAsyncOpHandle<FSendSessionInvite> FSessionsEOSGS::SendSessionInvite(FSend
 		}
 
 		TArray<TFuture<TDefaultErrorResult<FSendSingleSessionInviteImpl>>> PendingSessionInvites;
-		for (const FOnlineAccountIdHandle& TargetUser : OpParams.TargetUsers)
+		for (const FAccountId& TargetUser : OpParams.TargetUsers)
 		{
 			FSendSingleSessionInviteImpl::Params FSendSingleSessionInviteParams;
 			FSendSingleSessionInviteParams.LocalUserId = OpParams.LocalUserId;
@@ -1809,9 +1809,9 @@ TOnlineAsyncOpHandle<FAddSessionMembers> FSessionsEOSGS::AddSessionMembers(FAddS
 		Options.PlayersToRegisterCount = OpParams.NewSessionMembers.Num();
 
 		TArray<EOS_ProductUserId> PlayersToRegister;
-		TArray<FOnlineAccountIdHandle> SessionMemberIds;
+		TArray<FAccountId> SessionMemberIds;
 		OpParams.NewSessionMembers.GenerateKeyArray(SessionMemberIds);		
-		for (FOnlineAccountIdHandle IdHandle : SessionMemberIds)
+		for (FAccountId IdHandle : SessionMemberIds)
 		{
 			// TODO: Do this with ResolveAccountIds instead
 			PlayersToRegister.Add(GetProductUserIdChecked(IdHandle));
@@ -1879,7 +1879,7 @@ TOnlineAsyncOpHandle<FRemoveSessionMembers> FSessionsEOSGS::RemoveSessionMembers
 		Options.PlayersToUnregisterCount = OpParams.SessionMemberIds.Num();
 
 		TArray<EOS_ProductUserId> PlayersToUnregister;
-		for (FOnlineAccountIdHandle IdHandle : OpParams.SessionMemberIds)
+		for (FAccountId IdHandle : OpParams.SessionMemberIds)
 		{
 			// TODO: Do this with ResolveAccountIds instead
 			PlayersToUnregister.Add(GetProductUserIdChecked(IdHandle));
diff --git a/Engine/Plugins/Online/Experimental/OnlineServicesEOSGS/Source/Private/Online/StatsEOSGS.cpp b/Engine/Plugins/Online/Experimental/OnlineServicesEOSGS/Source/Private/Online/StatsEOSGS.cpp
--- a/Engine/Plugins/Online/Experimental/OnlineServicesEOSGS/Source/Private/Online/StatsEOSGS.cpp
+++ b/Engine/Plugins/Online/Experimental/OnlineServicesEOSGS/Source/Private/Online/StatsEOSGS.cpp
@@ -102,7 +102,7 @@ TOnlineAsyncOpHandle<FUpdateStats> FStatsEOSGS::UpdateStats(FUpdateStats::Params
 namespace Private
 {
 
-void QueryStatsEOS(EOS_HStats StatsHandle, const FOnlineAccountIdHandle& LocalUserId, const FOnlineAccountIdHandle& TargetUserId, const TArray<FString>& StatNames, TPromise<const EOS_Stats_OnQueryStatsCompleteCallbackInfo*>&& Promise)
+void QueryStatsEOS(EOS_HStats StatsHandle, const FAccountId& LocalUserId, const FAccountId& TargetUserId, const TArray<FString>& StatNames, TPromise<const EOS_Stats_OnQueryStatsCompleteCallbackInfo*>&& Promise)
 {
 	EOS_Stats_QueryStatsOptions Options;
 	Options.ApiVersion = EOS_STATS_QUERYSTATS_API_LATEST;
@@ -212,7 +212,7 @@ TOnlineAsyncOpHandle<FBatchQueryStats> FStatsEOSGS::BatchQueryStats(FBatchQueryS
 		return Op->GetHandle();
 	}
 
-	for (const FOnlineAccountIdHandle& TargetUserId : Op->GetParams().TargetUserIds)
+	for (const FAccountId& TargetUserId : Op->GetParams().TargetUserIds)
 	{
 		Op->Then([this, &TargetUserId](TOnlineAsyncOp<FBatchQueryStats>& InAsyncOp, TPromise<const EOS_Stats_OnQueryStatsCompleteCallbackInfo*>&& Promise)
 		{
diff --git a/Engine/Plugins/Online/Experimental/OnlineServicesEOSGS/Source/Public/Online/AchievementsEOSGS.h b/Engine/Plugins/Online/Experimental/OnlineServicesEOSGS/Source/Public/Online/AchievementsEOSGS.h
--- a/Engine/Plugins/Online/Experimental/OnlineServicesEOSGS/Source/Public/Online/AchievementsEOSGS.h
+++ b/Engine/Plugins/Online/Experimental/OnlineServicesEOSGS/Source/Public/Online/AchievementsEOSGS.h
@@ -41,7 +41,7 @@ protected:
 	using FAchievementStateMap = TMap<FString, FAchievementState>;
 
 	TOptional<FAchievementDefinitionMap> AchievementDefinitions;
-	TMap<FOnlineAccountIdHandle, FAchievementStateMap> AchievementStates;
+	TMap<FAccountId, FAchievementStateMap> AchievementStates;
 
 	EOS_NotificationId NotifyAchievementsUnlockedNotificationId = EOS_INVALID_NOTIFICATIONID;
 };
diff --git a/Engine/Plugins/Online/Experimental/OnlineServicesEOSGS/Source/Public/Online/AuthEOSGS.h b/Engine/Plugins/Online/Experimental/OnlineServicesEOSGS/Source/Public/Online/AuthEOSGS.h
--- a/Engine/Plugins/Online/Experimental/OnlineServicesEOSGS/Source/Public/Online/AuthEOSGS.h
+++ b/Engine/Plugins/Online/Experimental/OnlineServicesEOSGS/Source/Public/Online/AuthEOSGS.h
@@ -169,12 +169,12 @@ public:
 	virtual ~FAccountInfoRegistryEOS() = default;
 
 	TSharedPtr<FAccountInfoEOS> Find(FPlatformUserId PlatformUserId) const;
-	TSharedPtr<FAccountInfoEOS> Find(FOnlineAccountIdHandle AccountIdHandle) const;
+	TSharedPtr<FAccountInfoEOS> Find(FAccountId AccountIdHandle) const;
 	TSharedPtr<FAccountInfoEOS> Find(EOS_EpicAccountId EpicAccountId) const;
 	TSharedPtr<FAccountInfoEOS> Find(EOS_ProductUserId ProductUserId) const;
 
 	void Register(const TSharedRef<FAccountInfoEOS>& UserAuthData);
-	void Unregister(FOnlineAccountIdHandle AccountId);
+	void Unregister(FAccountId AccountId);
 
 protected:
 	virtual void DoRegister(const TSharedRef<FAccountInfo>& AccountInfo);
@@ -208,10 +208,10 @@ public:
 	// End IAuth
 
 	// Begin FAuthEOSGS
-	virtual TFuture<FOnlineAccountIdHandle> ResolveAccountId(const FOnlineAccountIdHandle& LocalUserId, const EOS_ProductUserId ProductUserId);
-	virtual TFuture<TArray<FOnlineAccountIdHandle>> ResolveAccountIds(const FOnlineAccountIdHandle& LocalUserId, const TArray<EOS_ProductUserId>& ProductUserIds);
-	virtual TFunction<TFuture<FOnlineAccountIdHandle>(FOnlineAsyncOp& InAsyncOp, const EOS_ProductUserId& ProductUserId)> ResolveProductIdFn();
-	virtual TFunction<TFuture<TArray<FOnlineAccountIdHandle>>(FOnlineAsyncOp& InAsyncOp, const TArray<EOS_ProductUserId>& ProductUserIds)> ResolveProductIdsFn();
+	virtual TFuture<FAccountId> ResolveAccountId(const FAccountId& LocalUserId, const EOS_ProductUserId ProductUserId);
+	virtual TFuture<TArray<FAccountId>> ResolveAccountIds(const FAccountId& LocalUserId, const TArray<EOS_ProductUserId>& ProductUserIds);
+	virtual TFunction<TFuture<FAccountId>(FOnlineAsyncOp& InAsyncOp, const EOS_ProductUserId& ProductUserId)> ResolveProductIdFn();
+	virtual TFunction<TFuture<TArray<FAccountId>>(FOnlineAsyncOp& InAsyncOp, const TArray<EOS_ProductUserId>& ProductUserIds)> ResolveProductIdsFn();
 	// End FAuthEOSGS
 
 protected:
@@ -243,7 +243,7 @@ protected:
 
 	void InitializeConnectLoginRecoveryTimer(const TSharedRef<FAccountInfoEOS>& UserAuthData);
 
-	static FOnlineAccountIdHandle CreateAccountId(const EOS_ProductUserId ProductUserId);
+	static FAccountId CreateAccountId(const EOS_ProductUserId ProductUserId);
 
 	EOS_HAuth AuthHandle = nullptr;
 	EOS_HConnect ConnectHandle = nullptr;
diff --git a/Engine/Plugins/Online/Experimental/OnlineServicesEOSGS/Source/Public/Online/LobbiesEOSGS.h b/Engine/Plugins/Online/Experimental/OnlineServicesEOSGS/Source/Public/Online/LobbiesEOSGS.h
--- a/Engine/Plugins/Online/Experimental/OnlineServicesEOSGS/Source/Public/Online/LobbiesEOSGS.h
+++ b/Engine/Plugins/Online/Experimental/OnlineServicesEOSGS/Source/Public/Online/LobbiesEOSGS.h
@@ -32,7 +32,7 @@ struct FLobbiesJoinLobbyImpl
 		TSharedPtr<FLobbyDataEOS> LobbyData;
 
 		// The local user agent which will perform the action.
-		FOnlineAccountIdHandle LocalUserId;
+		FAccountId LocalUserId;
 
 		// The local name for the lobby.
 		FName LocalName;
@@ -56,7 +56,7 @@ struct FLobbiesJoinLobbyMemberImpl
 		TSharedPtr<FLobbyDataEOS> LobbyData;
 
 		// The local user agent which will perform the action.
-		FOnlineAccountIdHandle LocalUserId;
+		FAccountId LocalUserId;
 
 		// Initial attributes.
 		TMap<FLobbyAttributeId, FLobbyVariant> Attributes;
@@ -77,7 +77,7 @@ struct FLobbiesLeaveLobbyImpl
 		TSharedPtr<FLobbyDataEOS> LobbyData;
 
 		// The local user agent which will perform the action.
-		FOnlineAccountIdHandle LocalUserId;
+		FAccountId LocalUserId;
 	};
 
 	struct Result
@@ -95,7 +95,7 @@ struct FLobbiesDestroyLobbyImpl
 		TSharedPtr<FLobbyDataEOS> LobbyData;
 
 		// The local user agent which will perform the action.
-		FOnlineAccountIdHandle LocalUserId;
+		FAccountId LocalUserId;
 	};
 
 	struct Result
@@ -113,10 +113,10 @@ struct FLobbiesInviteLobbyMemberImpl
 		TSharedPtr<FLobbyDataEOS> LobbyData;
 
 		// The local user agent which will perform the action.
-		FOnlineAccountIdHandle LocalUserId;
+		FAccountId LocalUserId;
 
 		// The target user for the invitation.
-		FOnlineAccountIdHandle TargetUserId;
+		FAccountId TargetUserId;
 	};
 
 	struct Result
@@ -131,7 +131,7 @@ struct FLobbiesDeclineLobbyInvitationImpl
 	struct Params
 	{
 		// The local user agent which will perform the action.
-		FOnlineAccountIdHandle LocalUserId;
+		FAccountId LocalUserId;
 
 		// Id of the lobby for which the invitations will be declined.
 		FOnlineLobbyIdHandle LobbyId;
@@ -152,10 +152,10 @@ struct FLobbiesKickLobbyMemberImpl
 		TSharedPtr<FLobbyDataEOS> LobbyData;
 
 		// The local user agent which will perform the action.
-		FOnlineAccountIdHandle LocalUserId;
+		FAccountId LocalUserId;
 
 		// The target user to be kicked.
-		FOnlineAccountIdHandle TargetUserId;
+		FAccountId TargetUserId;
 	};
 
 	struct Result
@@ -173,10 +173,10 @@ struct FLobbiesPromoteLobbyMemberImpl
 		TSharedPtr<FLobbyDataEOS> LobbyData;
 
 		// The local user agent which will perform the action.
-		FOnlineAccountIdHandle LocalUserId;
+		FAccountId LocalUserId;
 
 		// The target user to be promoted to owner.
-		FOnlineAccountIdHandle TargetUserId;
+		FAccountId TargetUserId;
 	};
 
 	struct Result
@@ -194,7 +194,7 @@ struct FLobbiesModifyLobbyDataImpl
 		TSharedPtr<FLobbyDataEOS> LobbyData;
 
 		// The local user agent which will perform the action.
-		FOnlineAccountIdHandle LocalUserId;
+		FAccountId LocalUserId;
 
 		// The changes to apply to the lobby data.
 		TSharedPtr<FClientLobbyDataChanges> Changes;
@@ -218,7 +218,7 @@ struct FLobbiesModifyLobbyMemberDataImpl
 		TSharedPtr<FLobbyDataEOS> LobbyData;
 
 		// The local user agent which will perform the action.
-		FOnlineAccountIdHandle LocalUserId;
+		FAccountId LocalUserId;
 
 		// The changes to apply to the lobby member data.
 		TSharedPtr<FClientLobbyMemberDataChanges> Changes;
@@ -293,13 +293,13 @@ protected:
 	void RegisterHandlers();
 	void UnregisterHandlers();
 
-	void AddActiveLobby(FOnlineAccountIdHandle LocalUserId, const TSharedRef<FLobbyDataEOS>& LobbyData);
-	void RemoveActiveLobby(FOnlineAccountIdHandle LocalUserId, const TSharedRef<FLobbyDataEOS>& LobbyData);
+	void AddActiveLobby(FAccountId LocalUserId, const TSharedRef<FLobbyDataEOS>& LobbyData);
+	void RemoveActiveLobby(FAccountId LocalUserId, const TSharedRef<FLobbyDataEOS>& LobbyData);
 
 	// Todo: store list of invites per lobby.
 	void AddActiveInvite(const TSharedRef<FLobbyInviteDataEOS>& Invite);
 	void RemoveActiveInvite(const TSharedRef<FLobbyInviteDataEOS>& Invite);
-	TSharedPtr<FLobbyInviteDataEOS> GetActiveInvite(FOnlineAccountIdHandle TargetUser, FOnlineLobbyIdHandle TargetLobbyId);
+	TSharedPtr<FLobbyInviteDataEOS> GetActiveInvite(FAccountId TargetUser, FOnlineLobbyIdHandle TargetLobbyId);
 
 	// LobbyData will be fetched from the operation data if not set in Params.
 	TFuture<TDefaultErrorResult<FLobbiesJoinLobbyImpl>> JoinLobbyImpl(FLobbiesJoinLobbyImpl::Params&& Params);
@@ -325,9 +325,9 @@ protected:
 	TSharedPtr<FLobbyPrerequisitesEOS> LobbyPrerequisites;
 	TSharedPtr<FLobbyDataRegistryEOS> LobbyDataRegistry;
 
-	TMap<FOnlineAccountIdHandle, TSet<TSharedRef<FLobbyDataEOS>>> ActiveLobbies;
-	TMap<FOnlineAccountIdHandle, TMap<FOnlineLobbyIdHandle, TSharedRef<FLobbyInviteDataEOS>>> ActiveInvites;
-	TMap<FOnlineAccountIdHandle, TSharedRef<FLobbySearchEOS>> ActiveSearchResults;
+	TMap<FAccountId, TSet<TSharedRef<FLobbyDataEOS>>> ActiveLobbies;
+	TMap<FAccountId, TMap<FOnlineLobbyIdHandle, TSharedRef<FLobbyInviteDataEOS>>> ActiveInvites;
+	TMap<FAccountId, TSharedRef<FLobbySearchEOS>> ActiveSearchResults;
 };
 
 namespace Meta {
diff --git a/Engine/Plugins/Online/Experimental/OnlineServicesEOSGS/Source/Public/Online/NboSerializerEOSGSSvc.h b/Engine/Plugins/Online/Experimental/OnlineServicesEOSGS/Source/Public/Online/NboSerializerEOSGSSvc.h
--- a/Engine/Plugins/Online/Experimental/OnlineServicesEOSGS/Source/Public/Online/NboSerializerEOSGSSvc.h
+++ b/Engine/Plugins/Online/Experimental/OnlineServicesEOSGS/Source/Public/Online/NboSerializerEOSGSSvc.h
@@ -19,7 +19,7 @@ namespace NboSerializerEOSGSSvc {
 
 /** NboSerializeToBuffer methods */
 
-inline void SerializeToBuffer(FNboSerializeToBuffer& Ar, const FOnlineAccountIdHandle& UniqueId)
+inline void SerializeToBuffer(FNboSerializeToBuffer& Ar, const FAccountId& UniqueId)
 {
 	TArray<uint8> Data = FOnlineAccountIdRegistryEOSGS::GetRegistered().ToReplicationData(UniqueId);
 	Ar << Data.Num();
@@ -37,7 +37,7 @@ inline void SerializeToBuffer(FNboSerializeToBuffer& Packet, const FSessionMembe
 {
 	Packet << SessionMembersMap.Num();
 
-	for (const TPair<FOnlineAccountIdHandle, FSessionMember>& Entry : SessionMembersMap)
+	for (const TPair<FAccountId, FSessionMember>& Entry : SessionMembersMap)
 	{
 		SerializeToBuffer(Packet, Entry.Key);
 		NboSerializerCommonSvc::SerializeToBuffer(Packet, Entry.Value);
@@ -48,7 +48,7 @@ inline void SerializeToBuffer(FNboSerializeToBuffer& Packet, const FRegisteredPl
 {
 	Packet << RegisteredPlayersMap.Num();
 
-	for (const TPair<FOnlineAccountIdHandle, FRegisteredPlayer>& Entry : RegisteredPlayersMap)
+	for (const TPair<FAccountId, FRegisteredPlayer>& Entry : RegisteredPlayersMap)
 	{
 		SerializeToBuffer(Packet, Entry.Key);
 		NboSerializerCommonSvc::SerializeToBuffer(Packet, Entry.Value);
@@ -57,7 +57,7 @@ inline void SerializeToBuffer(FNboSerializeToBuffer& Packet, const FRegisteredPl
 
 /** NboSerializeFromBuffer methods */
 
-inline void SerializeFromBuffer(FNboSerializeFromBuffer& Ar, FOnlineAccountIdHandle& UniqueId)
+inline void SerializeFromBuffer(FNboSerializeFromBuffer& Ar, FAccountId& UniqueId)
 {
 	TArray<uint8> Data;
 	int32 Size;
@@ -82,7 +82,7 @@ inline void SerializeFromBuffer(FNboSerializeFromBuffer& Packet, FSessionMembers
 
 	for (int32 Index = 0; Index < NumEntries; ++Index)
 	{
-		FOnlineAccountIdHandle Key;
+		FAccountId Key;
 		SerializeFromBuffer(Packet, Key);
 
 		FSessionMember Value;
@@ -99,7 +99,7 @@ inline void SerializeFromBuffer(FNboSerializeFromBuffer& Packet, FRegisteredPlay
 
 	for (int32 Index = 0; Index < NumEntries; ++Index)
 	{
-		FOnlineAccountIdHandle Key;
+		FAccountId Key;
 		SerializeFromBuffer(Packet, Key);
 
 		FRegisteredPlayer Value;
diff --git a/Engine/Plugins/Online/Experimental/OnlineServicesEOSGS/Source/Public/Online/OnlineIdEOSGS.h b/Engine/Plugins/Online/Experimental/OnlineServicesEOSGS/Source/Public/Online/OnlineIdEOSGS.h
--- a/Engine/Plugins/Online/Experimental/OnlineServicesEOSGS/Source/Public/Online/OnlineIdEOSGS.h
+++ b/Engine/Plugins/Online/Experimental/OnlineServicesEOSGS/Source/Public/Online/OnlineIdEOSGS.h
@@ -15,8 +15,8 @@ namespace UE::Online {
 class IOnlineAccountIdRegistryEOSGS : public IOnlineAccountIdRegistry
 {
 public:
-	virtual FOnlineAccountIdHandle FindAccountId(EOS_ProductUserId ProductUserId) const = 0;
-	virtual EOS_ProductUserId GetProductUserId(const FOnlineAccountIdHandle& Handle) const = 0;
+	virtual FAccountId FindAccountId(EOS_ProductUserId ProductUserId) const = 0;
+	virtual EOS_ProductUserId GetProductUserId(const FAccountId& Handle) const = 0;
 };
 
 /**
@@ -30,14 +30,14 @@ public:
 	virtual ~FOnlineAccountIdRegistryEOSGS() = default;
 
 	// Begin IOnlineAccountIdRegistryEOSGS
-	virtual FOnlineAccountIdHandle FindAccountId(const EOS_ProductUserId ProductUserId) const override;
-	virtual EOS_ProductUserId GetProductUserId(const FOnlineAccountIdHandle& Handle) const override;
+	virtual FAccountId FindAccountId(const EOS_ProductUserId ProductUserId) const override;
+	virtual EOS_ProductUserId GetProductUserId(const FAccountId& Handle) const override;
 	// End IOnlineAccountIdRegistryEOSGS
 
 	// Begin IOnlineAccountIdRegistry
-	virtual FString ToLogString(const FOnlineAccountIdHandle& Handle) const override;
-	virtual TArray<uint8> ToReplicationData(const FOnlineAccountIdHandle& Handle) const override;
-	virtual FOnlineAccountIdHandle FromReplicationData(const TArray<uint8>& ReplicationData) override;
+	virtual FString ToLogString(const FAccountId& Handle) const override;
+	virtual TArray<uint8> ToReplicationData(const FAccountId& Handle) const override;
+	virtual FAccountId FromReplicationData(const TArray<uint8>& ReplicationData) override;
 	// End IOnlineAccountIdRegistry
 
 	static IOnlineAccountIdRegistryEOSGS& GetRegistered();
@@ -47,15 +47,15 @@ private:
 	friend class FAuthEOSGS;
 	friend class FOnlineServicesEOSGSModule;
 	static FOnlineAccountIdRegistryEOSGS& Get();
-	FOnlineAccountIdHandle FindOrAddAccountId(const EOS_ProductUserId ProductUserId);
+	FAccountId FindOrAddAccountId(const EOS_ProductUserId ProductUserId);
 
 	TOnlineBasicAccountIdRegistry<EOS_ProductUserId> Registry;
 };
 
-EOS_ProductUserId ONLINESERVICESEOSGS_API GetProductUserId(const FOnlineAccountIdHandle& Handle);
-EOS_ProductUserId ONLINESERVICESEOSGS_API GetProductUserIdChecked(const FOnlineAccountIdHandle& Handle);
-FOnlineAccountIdHandle ONLINESERVICESEOSGS_API FindAccountId(const EOS_ProductUserId EpicAccountId);
-FOnlineAccountIdHandle ONLINESERVICESEOSGS_API FindAccountIdChecked(const EOS_ProductUserId EpicAccountId);
+EOS_ProductUserId ONLINESERVICESEOSGS_API GetProductUserId(const FAccountId& Handle);
+EOS_ProductUserId ONLINESERVICESEOSGS_API GetProductUserIdChecked(const FAccountId& Handle);
+FAccountId ONLINESERVICESEOSGS_API FindAccountId(const EOS_ProductUserId EpicAccountId);
+FAccountId ONLINESERVICESEOSGS_API FindAccountIdChecked(const EOS_ProductUserId EpicAccountId);
 
 template<typename IdType>
 inline bool ValidateOnlineId(const TOnlineIdHandle<IdType> Handle)
diff --git a/Engine/Plugins/Online/Experimental/OnlineServicesEOSGS/Source/Public/Online/SessionsEOSGS.h b/Engine/Plugins/Online/Experimental/OnlineServicesEOSGS/Source/Public/Online/SessionsEOSGS.h
--- a/Engine/Plugins/Online/Experimental/OnlineServicesEOSGS/Source/Public/Online/SessionsEOSGS.h
+++ b/Engine/Plugins/Online/Experimental/OnlineServicesEOSGS/Source/Public/Online/SessionsEOSGS.h
@@ -138,11 +138,11 @@ struct FSendSingleSessionInviteImpl
 
 	struct Params
 	{
-		FOnlineAccountIdHandle LocalUserId;
+		FAccountId LocalUserId;
 
 		FName SessionName;
 
-		FOnlineAccountIdHandle TargetUserId;
+		FAccountId TargetUserId;
 	};
 
 	struct Result
@@ -157,7 +157,7 @@ struct FBuildSessionFromDetailsHandle
 	struct Params
 	{
 		/** User which will drive the id resolution */
-		FOnlineAccountIdHandle LocalUserId;
+		FAccountId LocalUserId;
 
 		/** EOS session details handle used to extract the data */
 		TSharedRef<FSessionDetailsHandleEOSGS> SessionDetailsHandleEOSGS;
@@ -166,7 +166,7 @@ struct FBuildSessionFromDetailsHandle
 	struct Result
 	{
 		/** User which started the resolution operation */
-		FOnlineAccountIdHandle LocalUserId;
+		FAccountId LocalUserId;
 
 		/** Session built from the details handle */
 		TSharedRef<FSession> Session;
@@ -216,7 +216,7 @@ protected:
 	void SetSessionSearchMaxResults(FSessionSearchHandleEOSGS& SessionSearchHandle, uint32 MaxResults);
 	void SetSessionSearchParameters(FSessionSearchHandleEOSGS& SessionSearchHandle, TArray<FFindSessionsSearchFilter> Filters);
 	void SetSessionSearchSessionId(FSessionSearchHandleEOSGS& SessionSearchHandle, const FOnlineSessionIdHandle& SessionId);
-	void SetSessionSearchTargetId(FSessionSearchHandleEOSGS& SessionSearchHandle, const FOnlineAccountIdHandle& TargetUserId);
+	void SetSessionSearchTargetId(FSessionSearchHandleEOSGS& SessionSearchHandle, const FAccountId& TargetUserId);
 
 	/**
 	 * Writes all values in the passed SessionSettings to the SessionModificationHandle
diff --git a/Engine/Plugins/Online/Experimental/OnlineServicesEOSGS/Source/Public/Online/UserFileEOSGS.h b/Engine/Plugins/Online/Experimental/OnlineServicesEOSGS/Source/Public/Online/UserFileEOSGS.h
--- a/Engine/Plugins/Online/Experimental/OnlineServicesEOSGS/Source/Public/Online/UserFileEOSGS.h
+++ b/Engine/Plugins/Online/Experimental/OnlineServicesEOSGS/Source/Public/Online/UserFileEOSGS.h
@@ -52,7 +52,7 @@ protected:
 
 	FUserFileEOSGSConfig Config;
 
-	TMap<FOnlineAccountIdHandle, TArray<FString>> UserToFiles;
+	TMap<FAccountId, TArray<FString>> UserToFiles;
 
 	static EOS_PlayerDataStorage_EReadResult EOS_CALL OnReadFileDataStatic(const EOS_PlayerDataStorage_ReadFileDataCallbackInfo* Data);
 	static void EOS_CALL OnReadFileCompleteStatic(const EOS_PlayerDataStorage_ReadFileCallbackInfo* Data);
diff --git a/Engine/Plugins/Online/Experimental/OnlineServicesNull/Source/Private/Online/AuthNull.cpp b/Engine/Plugins/Online/Experimental/OnlineServicesNull/Source/Private/Online/AuthNull.cpp
--- a/Engine/Plugins/Online/Experimental/OnlineServicesNull/Source/Private/Online/AuthNull.cpp
+++ b/Engine/Plugins/Online/Experimental/OnlineServicesNull/Source/Private/Online/AuthNull.cpp
@@ -90,7 +90,7 @@ TSharedPtr<FAccountInfoNull> FAccountInfoRegistryNULL::Find(FPlatformUserId Plat
 	return StaticCastSharedPtr<FAccountInfoNull>(Super::Find(PlatformUserId));
 }
 
-TSharedPtr<FAccountInfoNull> FAccountInfoRegistryNULL::Find(FOnlineAccountIdHandle AccountIdHandle) const
+TSharedPtr<FAccountInfoNull> FAccountInfoRegistryNULL::Find(FAccountId AccountIdHandle) const
 {
 	return StaticCastSharedPtr<FAccountInfoNull>(Super::Find(AccountIdHandle));
 }
@@ -100,7 +100,7 @@ void FAccountInfoRegistryNULL::Register(const TSharedRef<FAccountInfoNull>& Acco
 	DoRegister(AccountInfoNULL);
 }
 
-void FAccountInfoRegistryNULL::Unregister(FOnlineAccountIdHandle AccountId)
+void FAccountInfoRegistryNULL::Unregister(FAccountId AccountId)
 {
 	if (TSharedPtr<FAccountInfoNull> AccountInfoNULL = Find(AccountId))
 	{
@@ -177,37 +177,37 @@ FOnlineAccountIdRegistryNull& FOnlineAccountIdRegistryNull::Get()
 	return Instance;
 }
 
-FOnlineAccountIdHandle FOnlineAccountIdRegistryNull::Find(FString UserId) const
+FAccountId FOnlineAccountIdRegistryNull::Find(FString UserId) const
 {
 	const FOnlineAccountIdString* Entry = StringToId.Find(UserId);
 	if(Entry)
 	{
 		return Entry->Handle;
 	}
-	return FOnlineAccountIdHandle();
+	return FAccountId();
 }
 
-FOnlineAccountIdHandle FOnlineAccountIdRegistryNull::Find(FPlatformUserId UserId) const
+FAccountId FOnlineAccountIdRegistryNull::Find(FPlatformUserId UserId) const
 {
 	const FOnlineAccountIdString* Entry = LocalUserMap.Find(UserId);
 	if (Entry)
 	{
 		return Entry->Handle;
 	}
-	return FOnlineAccountIdHandle();
+	return FAccountId();
 }
 
-FOnlineAccountIdHandle FOnlineAccountIdRegistryNull::Find(int32 UserIndex) const
+FAccountId FOnlineAccountIdRegistryNull::Find(int32 UserIndex) const
 {
 	const FOnlineAccountIdString* Entry = LocalUserMap.Find(FPlatformMisc::GetPlatformUserForUserIndex(UserIndex));
 	if (Entry)
 	{
 		return Entry->Handle;
 	}
-	return FOnlineAccountIdHandle();
+	return FAccountId();
 }
 
-const FOnlineAccountIdString* FOnlineAccountIdRegistryNull::GetInternal(const FOnlineAccountIdHandle& Handle) const
+const FOnlineAccountIdString* FOnlineAccountIdRegistryNull::GetInternal(const FAccountId& Handle) const
 {
 	if(Handle.IsValid() && Handle.GetOnlineServicesType() == EOnlineServices::Null && Handle.GetHandle() <= (uint32)Ids.Num())
 	{
@@ -216,9 +216,9 @@ const FOnlineAccountIdString* FOnlineAccountIdRegistryNull::GetInternal(const FO
 	return nullptr;
 }
 
-FOnlineAccountIdHandle FOnlineAccountIdRegistryNull::Create(FString UserId, FPlatformUserId PlatformUserId/* = PLATFORMUSERID_NONE*/)
+FAccountId FOnlineAccountIdRegistryNull::Create(FString UserId, FPlatformUserId PlatformUserId/* = PLATFORMUSERID_NONE*/)
 {
-	FOnlineAccountIdHandle ExistingHandle = Find(UserId);
+	FAccountId ExistingHandle = Find(UserId);
 	if(ExistingHandle.IsValid())
 	{
 		UE_LOG(LogOnlineServices, Error, TEXT("[FOnlineAccountIdRegistryNull::Create] Found a duplicate ID for local user %d."), FPlatformMisc::GetUserIndexForPlatformUser(PlatformUserId));
@@ -228,13 +228,13 @@ FOnlineAccountIdHandle FOnlineAccountIdRegistryNull::Create(FString UserId, FPla
 	if (!PlatformUserId.IsValid())
 	{
 		UE_LOG(LogOnlineServices, Warning, TEXT("[FOnlineAccountIdRegistryNull::Create] Unable to create id: PlatformUserId is invalid."));
-		return FOnlineAccountIdHandle();
+		return FAccountId();
 	}
 
 	FOnlineAccountIdString& Id = Ids.Emplace_GetRef();
 	Id.AccountIndex = Ids.Num();
 	Id.Data = UserId;
-	Id.Handle = FOnlineAccountIdHandle(EOnlineServices::Null, Id.AccountIndex);
+	Id.Handle = FAccountId(EOnlineServices::Null, Id.AccountIndex);
 	
 	StringToId.Add(UserId, Id);
 	LocalUserMap.Add(PlatformUserId, Id);
@@ -242,7 +242,7 @@ FOnlineAccountIdHandle FOnlineAccountIdRegistryNull::Create(FString UserId, FPla
 	return Id.Handle;
 }
 
-FString FOnlineAccountIdRegistryNull::ToLogString(const FOnlineAccountIdHandle& Handle) const
+FString FOnlineAccountIdRegistryNull::ToLogString(const FAccountId& Handle) const
 {
 	if(const FOnlineAccountIdString* Id = GetInternal(Handle))
 	{
@@ -253,7 +253,7 @@ FString FOnlineAccountIdRegistryNull::ToLogString(const FOnlineAccountIdHandle&
 }
 
 
-TArray<uint8> FOnlineAccountIdRegistryNull::ToReplicationData(const FOnlineAccountIdHandle& Handle) const
+TArray<uint8> FOnlineAccountIdRegistryNull::ToReplicationData(const FAccountId& Handle) const
 {
 	if (const FOnlineAccountIdString* Id = GetInternal(Handle))
 	{
@@ -267,14 +267,14 @@ TArray<uint8> FOnlineAccountIdRegistryNull::ToReplicationData(const FOnlineAccou
 	return TArray<uint8>();
 }
 
-FOnlineAccountIdHandle FOnlineAccountIdRegistryNull::FromReplicationData(const TArray<uint8>& ReplicationData)
+FAccountId FOnlineAccountIdRegistryNull::FromReplicationData(const TArray<uint8>& ReplicationData)
 {
 	FString Result = BytesToString(ReplicationData.GetData(), ReplicationData.Num());
 	if(Result.Len() > 0)
 	{
 		return Create(Result);
 	}
-	return FOnlineAccountIdHandle();
+	return FAccountId();
 }
 
 /* UE::Online */ }
diff --git a/Engine/Plugins/Online/Experimental/OnlineServicesNull/Source/Private/Online/LeaderboardsNull.cpp b/Engine/Plugins/Online/Experimental/OnlineServicesNull/Source/Private/Online/LeaderboardsNull.cpp
--- a/Engine/Plugins/Online/Experimental/OnlineServicesNull/Source/Private/Online/LeaderboardsNull.cpp
+++ b/Engine/Plugins/Online/Experimental/OnlineServicesNull/Source/Private/Online/LeaderboardsNull.cpp
@@ -159,7 +159,7 @@ TOnlineAsyncOpHandle<FReadEntriesForUsers> FLeaderboardsNull::ReadEntriesForUser
 
 		if (FLeaderboardDataNull* LeaderboardData = LeaderboardsData.FindByPredicate(FFindLeaderboardDataByName(InAsyncOp.GetParams().BoardName)))
 		{
-			for (const FOnlineAccountIdHandle& UserId : InAsyncOp.GetParams().UserIds)
+			for (const FAccountId& UserId : InAsyncOp.GetParams().UserIds)
 			{
 				uint32 Index = 0;
 				TDoubleLinkedList<FUserScoreNull>::TDoubleLinkedListNode* CurrentNode = LeaderboardData->UserScoreList.GetHead();
diff --git a/Engine/Plugins/Online/Experimental/OnlineServicesNull/Source/Private/Online/LobbiesNull.cpp b/Engine/Plugins/Online/Experimental/OnlineServicesNull/Source/Private/Online/LobbiesNull.cpp
--- a/Engine/Plugins/Online/Experimental/OnlineServicesNull/Source/Private/Online/LobbiesNull.cpp
+++ b/Engine/Plugins/Online/Experimental/OnlineServicesNull/Source/Private/Online/LobbiesNull.cpp
@@ -299,7 +299,7 @@ void FLobbiesNull::ReadLobbyFromPacket(FNboSerializeFromBuffer& Packet, const TS
 	UE_LOG(LogTemp, Verbose, TEXT("Reading session information from server"));
 
 	/** Owner of the session */
-	FOnlineAccountIdHandle OwningUserId;
+	FAccountId OwningUserId;
 	FString OwningUserName;
 	uint32 NumOpenPrivateConnections, NumOpenPublicConnections;
 	SerializeFromBuffer(Packet, OwningUserId);
diff --git a/Engine/Plugins/Online/Experimental/OnlineServicesNull/Source/Private/Online/PresenceNull.cpp b/Engine/Plugins/Online/Experimental/OnlineServicesNull/Source/Private/Online/PresenceNull.cpp
--- a/Engine/Plugins/Online/Experimental/OnlineServicesNull/Source/Private/Online/PresenceNull.cpp
+++ b/Engine/Plugins/Online/Experimental/OnlineServicesNull/Source/Private/Online/PresenceNull.cpp
@@ -73,7 +73,7 @@ TOnlineAsyncOpHandle<FUpdatePresence> FPresenceNull::UpdatePresence(FUpdatePrese
 		const FUpdatePresence::Params& Params = InAsyncOp.GetParams();
 		Presences.Add(Params.LocalUserId, Params.Presence);
 
-		for (const TPair<FOnlineAccountIdHandle, TSet<FOnlineAccountIdHandle>>& Pairs : PresenceListeners)
+		for (const TPair<FAccountId, TSet<FAccountId>>& Pairs : PresenceListeners)
 		{
 			if (Pairs.Value.Contains(Params.LocalUserId))
 			{
@@ -106,7 +106,7 @@ TOnlineAsyncOpHandle<FPartialUpdatePresence> FPresenceNull::PartialUpdatePresenc
 		TSharedRef<FUserPresence> MutatedPresence = ApplyPresenceMutations(*NewPresence, Mutations);
 		Presences.Add(InAsyncOp.GetParams().LocalUserId, MutatedPresence);
 
-		for (const TPair<FOnlineAccountIdHandle, TSet<FOnlineAccountIdHandle>>& Pairs : PresenceListeners)
+		for (const TPair<FAccountId, TSet<FAccountId>>& Pairs : PresenceListeners)
 		{
 			if (Pairs.Value.Contains(InAsyncOp.GetParams().LocalUserId))
 			{
diff --git a/Engine/Plugins/Online/Experimental/OnlineServicesNull/Source/Private/Online/StatsNull.cpp b/Engine/Plugins/Online/Experimental/OnlineServicesNull/Source/Private/Online/StatsNull.cpp
--- a/Engine/Plugins/Online/Experimental/OnlineServicesNull/Source/Private/Online/StatsNull.cpp
+++ b/Engine/Plugins/Online/Experimental/OnlineServicesNull/Source/Private/Online/StatsNull.cpp
@@ -139,7 +139,7 @@ TOnlineAsyncOpHandle<FBatchQueryStats> FStatsNull::BatchQueryStats(FBatchQuerySt
 	{
 		FBatchQueryStats::Result Result;
 
-		for (const FOnlineAccountIdHandle& TargetUserId : InAsyncOp.GetParams().TargetUserIds)
+		for (const FAccountId& TargetUserId : InAsyncOp.GetParams().TargetUserIds)
 		{
 			if (FUserStats* ExistingUserStats = UsersStats.FindByPredicate(FFindUserStatsByAccountId(TargetUserId)))
 			{
diff --git a/Engine/Plugins/Online/Experimental/OnlineServicesNull/Source/Private/Online/UserFileNull.cpp b/Engine/Plugins/Online/Experimental/OnlineServicesNull/Source/Private/Online/UserFileNull.cpp
--- a/Engine/Plugins/Online/Experimental/OnlineServicesNull/Source/Private/Online/UserFileNull.cpp
+++ b/Engine/Plugins/Online/Experimental/OnlineServicesNull/Source/Private/Online/UserFileNull.cpp
@@ -191,7 +191,7 @@ TOnlineAsyncOpHandle<FUserFileDeleteFile> FUserFileNull::DeleteFile(FUserFileDel
 	return Op->GetHandle();
 }
 
-FUserFileNull::FUserState& FUserFileNull::GetUserState(const FOnlineAccountIdHandle AccountId)
+FUserFileNull::FUserState& FUserFileNull::GetUserState(const FAccountId AccountId)
 {
 	check(AccountId.IsValid());
 
diff --git a/Engine/Plugins/Online/Experimental/OnlineServicesNull/Source/Public/Online/AchievementsNull.h b/Engine/Plugins/Online/Experimental/OnlineServicesNull/Source/Public/Online/AchievementsNull.h
--- a/Engine/Plugins/Online/Experimental/OnlineServicesNull/Source/Public/Online/AchievementsNull.h
+++ b/Engine/Plugins/Online/Experimental/OnlineServicesNull/Source/Public/Online/AchievementsNull.h
@@ -30,7 +30,7 @@ protected:
 	using FAchievementStateMap = TMap<FString, FAchievementState>;
 
 	TOptional<FAchievementDefinitionMap> AchievementDefinitions;
-	TMap<FOnlineAccountIdHandle, FAchievementStateMap> AchievementStates;
+	TMap<FAccountId, FAchievementStateMap> AchievementStates;
 };
 
 /* UE::Online */ }
diff --git a/Engine/Plugins/Online/Experimental/OnlineServicesNull/Source/Public/Online/AuthNull.h b/Engine/Plugins/Online/Experimental/OnlineServicesNull/Source/Public/Online/AuthNull.h
--- a/Engine/Plugins/Online/Experimental/OnlineServicesNull/Source/Public/Online/AuthNull.h
+++ b/Engine/Plugins/Online/Experimental/OnlineServicesNull/Source/Public/Online/AuthNull.h
@@ -14,7 +14,7 @@ class FOnlineAccountIdString
 public:
 	FString Data;
 	int32 AccountIndex;
-	FOnlineAccountIdHandle Handle;
+	FAccountId Handle;
 };
 
 class FOnlineAccountIdRegistryNull : public IOnlineAccountIdRegistry
@@ -22,22 +22,22 @@ class FOnlineAccountIdRegistryNull : public IOnlineAccountIdRegistry
 public:
 	static FOnlineAccountIdRegistryNull& Get();
 
-	FOnlineAccountIdHandle Find(FString UserId) const;
-	FOnlineAccountIdHandle Find(FPlatformUserId UserId) const;
-	FOnlineAccountIdHandle Find(int32 UserIndex) const;
+	FAccountId Find(FString UserId) const;
+	FAccountId Find(FPlatformUserId UserId) const;
+	FAccountId Find(int32 UserIndex) const;
 
-	FOnlineAccountIdHandle Create(FString UserId, FPlatformUserId PlatformUserId = PLATFORMUSERID_NONE);
+	FAccountId Create(FString UserId, FPlatformUserId PlatformUserId = PLATFORMUSERID_NONE);
 
 	// Begin IOnlineAccountIdRegistry
-	virtual FString ToLogString(const FOnlineAccountIdHandle& Handle) const override;
-	virtual TArray<uint8> ToReplicationData(const FOnlineAccountIdHandle& Handle) const override;
-	virtual FOnlineAccountIdHandle FromReplicationData(const TArray<uint8>& ReplicationString) override;
+	virtual FString ToLogString(const FAccountId& Handle) const override;
+	virtual TArray<uint8> ToReplicationData(const FAccountId& Handle) const override;
+	virtual FAccountId FromReplicationData(const TArray<uint8>& ReplicationString) override;
 	// End IOnlineAccountIdRegistry
 
 	virtual ~FOnlineAccountIdRegistryNull() = default;
 
 private:
-	const FOnlineAccountIdString* GetInternal(const FOnlineAccountIdHandle& Handle) const;
+	const FOnlineAccountIdString* GetInternal(const FAccountId& Handle) const;
 
 	// how much of these are actually necessary?
 	TArray<FOnlineAccountIdString> Ids;
@@ -62,10 +62,10 @@ public:
 	virtual ~FAccountInfoRegistryNULL() = default;
 
 	TSharedPtr<FAccountInfoNull> Find(FPlatformUserId PlatformUserId) const;
-	TSharedPtr<FAccountInfoNull> Find(FOnlineAccountIdHandle AccountIdHandle) const;
+	TSharedPtr<FAccountInfoNull> Find(FAccountId AccountIdHandle) const;
 
 	void Register(const TSharedRef<FAccountInfoNull>&UserAuthData);
-	void Unregister(FOnlineAccountIdHandle AccountId);
+	void Unregister(FAccountId AccountId);
 };
 
 class ONLINESERVICESNULL_API FAuthNull : public FAuthCommon
diff --git a/Engine/Plugins/Online/Experimental/OnlineServicesNull/Source/Public/Online/LeaderboardsNull.h b/Engine/Plugins/Online/Experimental/OnlineServicesNull/Source/Public/Online/LeaderboardsNull.h
--- a/Engine/Plugins/Online/Experimental/OnlineServicesNull/Source/Public/Online/LeaderboardsNull.h
+++ b/Engine/Plugins/Online/Experimental/OnlineServicesNull/Source/Public/Online/LeaderboardsNull.h
@@ -12,7 +12,7 @@ class FOnlineServicesNull;
 
 struct FUserScoreNull
 {
-	FOnlineAccountIdHandle UserId;
+	FAccountId UserId;
 	uint64 Score;
 };
 
diff --git a/Engine/Plugins/Online/Experimental/OnlineServicesNull/Source/Public/Online/NboSerializerNullSvc.h b/Engine/Plugins/Online/Experimental/OnlineServicesNull/Source/Public/Online/NboSerializerNullSvc.h
--- a/Engine/Plugins/Online/Experimental/OnlineServicesNull/Source/Public/Online/NboSerializerNullSvc.h
+++ b/Engine/Plugins/Online/Experimental/OnlineServicesNull/Source/Public/Online/NboSerializerNullSvc.h
@@ -19,7 +19,7 @@ namespace NboSerializerNullSvc {
 
 /** NboSerializeToBuffer methods */
 
-inline void SerializeToBuffer(FNboSerializeToBuffer& Ar, const FOnlineAccountIdHandle& UniqueId)
+inline void SerializeToBuffer(FNboSerializeToBuffer& Ar, const FAccountId& UniqueId)
 {
 	TArray<uint8> Data = FOnlineAccountIdRegistryNull::Get().ToReplicationData(UniqueId);
 	Ar << Data.Num();
@@ -37,7 +37,7 @@ inline void SerializeToBuffer(FNboSerializeToBuffer& Packet, const FSessionMembe
 {
 	Packet << SessionMembersMap.Num();
 
-	for (const TPair<FOnlineAccountIdHandle, FSessionMember>& Entry : SessionMembersMap)
+	for (const TPair<FAccountId, FSessionMember>& Entry : SessionMembersMap)
 	{
 		SerializeToBuffer(Packet, Entry.Key);
 		NboSerializerCommonSvc::SerializeToBuffer(Packet, Entry.Value);
@@ -48,7 +48,7 @@ inline void SerializeToBuffer(FNboSerializeToBuffer& Packet, const FRegisteredPl
 {
 	Packet << RegisteredPlayersMap.Num();
 
-	for (const TPair<FOnlineAccountIdHandle, FRegisteredPlayer>& Entry : RegisteredPlayersMap)
+	for (const TPair<FAccountId, FRegisteredPlayer>& Entry : RegisteredPlayersMap)
 	{
 		SerializeToBuffer(Packet, Entry.Key);
 		NboSerializerCommonSvc::SerializeToBuffer(Packet, Entry.Value);
@@ -67,7 +67,7 @@ inline void SerializeToBuffer(FNboSerializeToBuffer& Ar, const TMap<FLobbyAttrib
 
 /** NboSerializeFromBuffer methods */
 
-inline void SerializeFromBuffer(FNboSerializeFromBuffer& Ar, FOnlineAccountIdHandle& UniqueId)
+inline void SerializeFromBuffer(FNboSerializeFromBuffer& Ar, FAccountId& UniqueId)
 {
 	TArray<uint8> Data;
 	int32 Size;
@@ -94,7 +94,7 @@ inline void SerializeFromBuffer(FNboSerializeFromBuffer& Packet, FSessionMembers
 
 	for (int32 Index = 0; Index < NumEntries; ++Index)
 	{
-		FOnlineAccountIdHandle Key;
+		FAccountId Key;
 		SerializeFromBuffer(Packet, Key);
 
 		FSessionMember Value;
@@ -111,7 +111,7 @@ inline void SerializeFromBuffer(FNboSerializeFromBuffer& Packet, FRegisteredPlay
 
 	for (int32 Index = 0; Index < NumEntries; ++Index)
 	{
-		FOnlineAccountIdHandle Key;
+		FAccountId Key;
 		SerializeFromBuffer(Packet, Key);
 
 		FRegisteredPlayer Value;
diff --git a/Engine/Plugins/Online/Experimental/OnlineServicesNull/Source/Public/Online/PresenceNull.h b/Engine/Plugins/Online/Experimental/OnlineServicesNull/Source/Public/Online/PresenceNull.h
--- a/Engine/Plugins/Online/Experimental/OnlineServicesNull/Source/Public/Online/PresenceNull.h
+++ b/Engine/Plugins/Online/Experimental/OnlineServicesNull/Source/Public/Online/PresenceNull.h
@@ -23,8 +23,8 @@ public:
 	virtual TOnlineAsyncOpHandle<FPartialUpdatePresence> PartialUpdatePresence(FPartialUpdatePresence::Params&& Params) override;
 
 protected:
-	TMap<FOnlineAccountIdHandle, TSharedRef<const FUserPresence>> Presences;
-	TMap<FOnlineAccountIdHandle, TSet<FOnlineAccountIdHandle>> PresenceListeners;
+	TMap<FAccountId, TSharedRef<const FUserPresence>> Presences;
+	TMap<FAccountId, TSet<FAccountId>> PresenceListeners;
 };
 
 	/* UE::Online */
diff --git a/Engine/Plugins/Online/Experimental/OnlineServicesNull/Source/Public/Online/UserFileNull.h b/Engine/Plugins/Online/Experimental/OnlineServicesNull/Source/Public/Online/UserFileNull.h
--- a/Engine/Plugins/Online/Experimental/OnlineServicesNull/Source/Public/Online/UserFileNull.h
+++ b/Engine/Plugins/Online/Experimental/OnlineServicesNull/Source/Public/Online/UserFileNull.h
@@ -36,8 +36,8 @@ protected:
 		FUserFileMap Files;
 	};
 
-	TMap<FOnlineAccountIdHandle, FUserState> UserStates;
-	FUserState& GetUserState(const FOnlineAccountIdHandle AccountId);
+	TMap<FAccountId, FUserState> UserStates;
+	FUserState& GetUserState(const FAccountId AccountId);
 
 	FUserFileMap InitialFileStateFromConfig;
 };
diff --git a/Engine/Plugins/Online/Experimental/OnlineServicesOSSAdapter/Source/Private/Online/AchievementsOSSAdapter.cpp b/Engine/Plugins/Online/Experimental/OnlineServicesOSSAdapter/Source/Private/Online/AchievementsOSSAdapter.cpp
--- a/Engine/Plugins/Online/Experimental/OnlineServicesOSSAdapter/Source/Private/Online/AchievementsOSSAdapter.cpp
+++ b/Engine/Plugins/Online/Experimental/OnlineServicesOSSAdapter/Source/Private/Online/AchievementsOSSAdapter.cpp
@@ -32,7 +32,7 @@ void FAchievementsOSSAdapter::PostInitialize()
 
 	AchievementsInterface->AddOnAchievementUnlockedDelegate_Handle(FOnAchievementUnlockedDelegate::CreateLambda([this](const FUniqueNetId& LocalUserIdV1, const FString& AchievementId)
 	{
-		const FOnlineAccountIdHandle LocalUserIdV2 = Services.Get<FAuthOSSAdapter>()->GetAccountIdHandle(LocalUserIdV1.AsShared());
+		const FAccountId LocalUserIdV2 = Services.Get<FAuthOSSAdapter>()->GetAccountIdHandle(LocalUserIdV1.AsShared());
 
 		FAchievementStateUpdated AchievementStateUpdated;
 		AchievementStateUpdated.LocalUserId = LocalUserIdV2;
diff --git a/Engine/Plugins/Online/Experimental/OnlineServicesOSSAdapter/Source/Private/Online/AuthOSSAdapter.cpp b/Engine/Plugins/Online/Experimental/OnlineServicesOSSAdapter/Source/Private/Online/AuthOSSAdapter.cpp
--- a/Engine/Plugins/Online/Experimental/OnlineServicesOSSAdapter/Source/Private/Online/AuthOSSAdapter.cpp
+++ b/Engine/Plugins/Online/Experimental/OnlineServicesOSSAdapter/Source/Private/Online/AuthOSSAdapter.cpp
@@ -100,7 +100,7 @@ TSharedPtr<FAccountInfoOSSAdapter> FAccountInfoRegistryOSSAdapter::Find(FPlatfor
 	return StaticCastSharedPtr<FAccountInfoOSSAdapter>(Super::Find(PlatformUserId));
 } 
 
-TSharedPtr<FAccountInfoOSSAdapter> FAccountInfoRegistryOSSAdapter::Find(FOnlineAccountIdHandle AccountIdHandle) const
+TSharedPtr<FAccountInfoOSSAdapter> FAccountInfoRegistryOSSAdapter::Find(FAccountId AccountIdHandle) const
 {
 	return StaticCastSharedPtr<FAccountInfoOSSAdapter>(Super::Find(AccountIdHandle));
 }
@@ -110,7 +110,7 @@ void FAccountInfoRegistryOSSAdapter::Register(const TSharedRef<FAccountInfoOSSAd
 	DoRegister(AccountInfoNULL);
 }
 
-void FAccountInfoRegistryOSSAdapter::Unregister(FOnlineAccountIdHandle AccountId)
+void FAccountInfoRegistryOSSAdapter::Unregister(FAccountId AccountId)
 {
 	if (TSharedPtr<FAccountInfoOSSAdapter> AccountInfoNULL = Find(AccountId))
 	{
@@ -516,17 +516,17 @@ TOnlineAsyncOpHandle<FAuthQueryExternalAuthToken> FAuthOSSAdapter::QueryExternal
 	return Op->GetHandle();
 }
 
-FUniqueNetIdPtr FAuthOSSAdapter::GetUniqueNetId(FOnlineAccountIdHandle AccountIdHandle) const
+FUniqueNetIdPtr FAuthOSSAdapter::GetUniqueNetId(FAccountId AccountIdHandle) const
 {
 	return GetOnlineServicesOSSAdapter().GetAccountIdRegistry().GetIdValue(AccountIdHandle);
 }
 
-FOnlineAccountIdHandle FAuthOSSAdapter::GetAccountIdHandle(const FUniqueNetIdRef& UniqueNetId) const
+FAccountId FAuthOSSAdapter::GetAccountIdHandle(const FUniqueNetIdRef& UniqueNetId) const
 {
 	return GetOnlineServicesOSSAdapter().GetAccountIdRegistry().FindOrAddHandle(UniqueNetId);
 }
 
-int32 FAuthOSSAdapter::GetLocalUserNum(FOnlineAccountIdHandle AccountIdHandle) const
+int32 FAuthOSSAdapter::GetLocalUserNum(FAccountId AccountIdHandle) const
 {
 	TSharedPtr<FAccountInfoOSSAdapter> AccountInfoOSSAdapter = AccountInfoRegistryOSSAdapter.Find(AccountIdHandle);
 	return AccountInfoOSSAdapter ? AccountInfoOSSAdapter->LocalUserNum : INDEX_NONE;
diff --git a/Engine/Plugins/Online/Experimental/OnlineServicesOSSAdapter/Source/Private/Online/ExternalUIOSSAdapter.cpp b/Engine/Plugins/Online/Experimental/OnlineServicesOSSAdapter/Source/Private/Online/ExternalUIOSSAdapter.cpp
--- a/Engine/Plugins/Online/Experimental/OnlineServicesOSSAdapter/Source/Private/Online/ExternalUIOSSAdapter.cpp
+++ b/Engine/Plugins/Online/Experimental/OnlineServicesOSSAdapter/Source/Private/Online/ExternalUIOSSAdapter.cpp
@@ -66,7 +66,7 @@ TOnlineAsyncOpHandle<FExternalUIShowLoginUI> FExternalUIOSSAdapter::ShowLoginUI(
 							{
 								if (Error.WasSuccessful())
 								{
-									FOnlineAccountIdHandle Handle = static_cast<FOnlineServicesOSSAdapter&>(Services).GetAccountIdRegistry().FindOrAddHandle(UniqueId.ToSharedRef());
+									FAccountId Handle = static_cast<FOnlineServicesOSSAdapter&>(Services).GetAccountIdRegistry().FindOrAddHandle(UniqueId.ToSharedRef());
 									FExternalUIShowLoginUI::Result Result = { MakeShared<FAccountInfo>() };
 									Result.AccountInfo->PlatformUserId = PinnedOp->GetParams().PlatformUserId;
 									Result.AccountInfo->AccountId = Handle;
diff --git a/Engine/Plugins/Online/Experimental/OnlineServicesOSSAdapter/Source/Private/Online/LeaderboardsOSSAdapter.cpp b/Engine/Plugins/Online/Experimental/OnlineServicesOSSAdapter/Source/Private/Online/LeaderboardsOSSAdapter.cpp
--- a/Engine/Plugins/Online/Experimental/OnlineServicesOSSAdapter/Source/Private/Online/LeaderboardsOSSAdapter.cpp
+++ b/Engine/Plugins/Online/Experimental/OnlineServicesOSSAdapter/Source/Private/Online/LeaderboardsOSSAdapter.cpp
@@ -42,7 +42,7 @@ TOnlineAsyncOpHandle<FReadEntriesForUsers> FLeaderboardsOSSAdapter::ReadEntriesF
 		}
 
 		TArray<FUniqueNetIdRef> NetIds;
-		for (const FOnlineAccountIdHandle& UserId : Op.GetParams().UserIds)
+		for (const FAccountId& UserId : Op.GetParams().UserIds)
 		{
 			const FUniqueNetIdPtr UniqueNetId = Auth->GetUniqueNetId(UserId);
 			if (!UniqueNetId)
diff --git a/Engine/Plugins/Online/Experimental/OnlineServicesOSSAdapter/Source/Private/Online/PresenceOSSAdapter.cpp b/Engine/Plugins/Online/Experimental/OnlineServicesOSSAdapter/Source/Private/Online/PresenceOSSAdapter.cpp
--- a/Engine/Plugins/Online/Experimental/OnlineServicesOSSAdapter/Source/Private/Online/PresenceOSSAdapter.cpp
+++ b/Engine/Plugins/Online/Experimental/OnlineServicesOSSAdapter/Source/Private/Online/PresenceOSSAdapter.cpp
@@ -187,7 +187,7 @@ TOnlineAsyncOpHandle<FBatchQueryPresence> FPresenceOSSAdapter::BatchQueryPresenc
 		}
 
 		TArray<FUniqueNetIdRef> NetIds;
-		for (const FOnlineAccountIdHandle& TargetUserId : Op.GetParams().TargetUserIds)
+		for (const FAccountId& TargetUserId : Op.GetParams().TargetUserIds)
 		{
 			const FUniqueNetIdPtr TargetUserNetId = Auth->GetUniqueNetId(TargetUserId);
 			if (!TargetUserNetId)
diff --git a/Engine/Plugins/Online/Experimental/OnlineServicesOSSAdapter/Source/Private/Online/SessionsOSSAdapter.cpp b/Engine/Plugins/Online/Experimental/OnlineServicesOSSAdapter/Source/Private/Online/SessionsOSSAdapter.cpp
--- a/Engine/Plugins/Online/Experimental/OnlineServicesOSSAdapter/Source/Private/Online/SessionsOSSAdapter.cpp
+++ b/Engine/Plugins/Online/Experimental/OnlineServicesOSSAdapter/Source/Private/Online/SessionsOSSAdapter.cpp
@@ -337,7 +337,7 @@ void FSessionsOSSAdapter::Initialize()
 
 			const FOnlineServicesOSSAdapter& ServicesOSSAdapter = static_cast<FOnlineServicesOSSAdapter&>(Services);
 			const FUniqueNetIdRef TargetUniqueNetIdRef = TargetUniqueNetId.AsShared();
-			const FOnlineAccountIdHandle TargetAccountIdHandle = ServicesOSSAdapter.GetAccountIdRegistry().FindOrAddHandle(TargetUniqueNetIdRef);
+			const FAccountId TargetAccountIdHandle = ServicesOSSAdapter.GetAccountIdRegistry().FindOrAddHandle(TargetUniqueNetIdRef);
 
 			FSessionSettingsUpdate SettingsUpdate;
 			SettingsUpdate.RemovedSessionMembers.Add(TargetAccountIdHandle);
@@ -358,7 +358,7 @@ void FSessionsOSSAdapter::Initialize()
 
 			const FOnlineServicesOSSAdapter& ServicesOSSAdapter = static_cast<FOnlineServicesOSSAdapter&>(Services);
 			const FUniqueNetIdRef TargetUniqueNetIdRef = TargetUniqueNetId.AsShared();
-			const FOnlineAccountIdHandle TargetAccountIdHandle = ServicesOSSAdapter.GetAccountIdRegistry().FindOrAddHandle(TargetUniqueNetIdRef);
+			const FAccountId TargetAccountIdHandle = ServicesOSSAdapter.GetAccountIdRegistry().FindOrAddHandle(TargetUniqueNetIdRef);
 
 			FSessionSettingsUpdate SettingsUpdate;
 			if (bJoined)
@@ -392,7 +392,7 @@ void FSessionsOSSAdapter::Initialize()
 
 			const FOnlineServicesOSSAdapter& ServicesOSSAdapter = static_cast<FOnlineServicesOSSAdapter&>(Services);
 			const FUniqueNetIdRef TargetUniqueNetIdRef = TargetUniqueNetId.AsShared();
-			const FOnlineAccountIdHandle TargetAccountIdHandle = ServicesOSSAdapter.GetAccountIdRegistry().FindOrAddHandle(TargetUniqueNetIdRef);
+			const FAccountId TargetAccountIdHandle = ServicesOSSAdapter.GetAccountIdRegistry().FindOrAddHandle(TargetUniqueNetIdRef);
 
 			FSessionSettingsUpdate SettingsUpdate;
 			if (const FSessionMember* SessionMember = FoundSession->SessionSettings.SessionMembers.Find(TargetAccountIdHandle))
@@ -841,7 +841,7 @@ TOnlineAsyncOpHandle<FStartMatchmaking> FSessionsOSSAdapter::StartMatchmaking(FS
 		FOnlineServicesOSSAdapter& ServicesOSSAdapter = static_cast<FOnlineServicesOSSAdapter&>(Services);
 
 		TArray<FSessionMatchmakingUser> MatchMakingUsers;
-		for (const TPair<FOnlineAccountIdHandle, FSessionMember>& LocalUser : OpParams.LocalUsers)
+		for (const TPair<FAccountId, FSessionMember>& LocalUser : OpParams.LocalUsers)
 		{
 			FSessionMatchmakingUser NewMatchMakingUser { ServicesOSSAdapter.GetAccountIdRegistry().GetIdValue(LocalUser.Key).ToSharedRef() };
 
@@ -1039,9 +1039,9 @@ TOnlineAsyncOpHandle<FAddSessionMembers> FSessionsOSSAdapter::AddSessionMembers(
 
 		FAuthOSSAdapter* Auth = Services.Get<FAuthOSSAdapter>();
 		TArray<FUniqueNetIdRef> TargetUserNetIds;
-		TArray<FOnlineAccountIdHandle> NewSessionMemberIds;
+		TArray<FAccountId> NewSessionMemberIds;
 		OpParams.NewSessionMembers.GenerateKeyArray(NewSessionMemberIds);
-		for (const FOnlineAccountIdHandle& TargetUser : NewSessionMemberIds)
+		for (const FAccountId& TargetUser : NewSessionMemberIds)
 		{
 			FUniqueNetIdRef TargetUserNetId = Auth->GetUniqueNetId(TargetUser).ToSharedRef();
 			if (TargetUserNetId->IsValid())
@@ -1094,7 +1094,7 @@ TOnlineAsyncOpHandle<FRemoveSessionMembers> FSessionsOSSAdapter::RemoveSessionMe
 
 		FAuthOSSAdapter* Auth = Services.Get<FAuthOSSAdapter>();
 		TArray<FUniqueNetIdRef> TargetUserNetIds;
-		for (const FOnlineAccountIdHandle& TargetUser : OpParams.SessionMemberIds)
+		for (const FAccountId& TargetUser : OpParams.SessionMemberIds)
 		{
 			FUniqueNetIdRef TargetUserNetId = Auth->GetUniqueNetId(TargetUser).ToSharedRef();
 			if (TargetUserNetId->IsValid())
@@ -1136,7 +1136,7 @@ TOnlineAsyncOpHandle<FSendSessionInvite> FSessionsOSSAdapter::SendSessionInvite(
 		}
 
 		TArray<FUniqueNetIdRef> TargetUserNetIds;
-		for (const FOnlineAccountIdHandle& TargetUser : OpParams.TargetUsers)
+		for (const FAccountId& TargetUser : OpParams.TargetUsers)
 		{
 			FUniqueNetIdRef TargetUserNetId = Auth->GetUniqueNetId(TargetUser).ToSharedRef();
 			if (TargetUserNetId->IsValid())
@@ -1241,7 +1241,7 @@ FOnlineSessionSettings FSessionsOSSAdapter::BuildV1Settings(const FSessionSettin
 	
 	FOnlineServicesOSSAdapter& ServicesOSSAdapter = static_cast<FOnlineServicesOSSAdapter&>(Services);
 
-	for (const TPair<FOnlineAccountIdHandle, FSessionMember>& Entry : InSessionSettings.SessionMembers)
+	for (const TPair<FAccountId, FSessionMember>& Entry : InSessionSettings.SessionMembers)
 	{
 		Result.MemberSettings.Emplace(ServicesOSSAdapter.GetAccountIdRegistry().GetIdValue(Entry.Key).ToSharedRef(), ::FSessionSettings()); // TODO: Pending SchemaVariant work
 	}
@@ -1276,7 +1276,7 @@ void FSessionsOSSAdapter::WriteV2SessionSettingsFromV1Session(const FOnlineSessi
 
 	for (const TPair<FUniqueNetIdRef, ::FSessionSettings>& Entry : InSession->SessionSettings.MemberSettings)
 	{
-		FOnlineAccountIdHandle SessionMemberId = ServicesOSSAdapter.GetAccountIdRegistry().FindOrAddHandle(Entry.Key);
+		FAccountId SessionMemberId = ServicesOSSAdapter.GetAccountIdRegistry().FindOrAddHandle(Entry.Key);
 
 		FSessionMember SessionMember;
 		SessionMember.MemberSettings = GetV2SessionSettings(Entry.Value);
diff --git a/Engine/Plugins/Online/Experimental/OnlineServicesOSSAdapter/Source/Private/Online/StatsOSSAdapter.cpp b/Engine/Plugins/Online/Experimental/OnlineServicesOSSAdapter/Source/Private/Online/StatsOSSAdapter.cpp
--- a/Engine/Plugins/Online/Experimental/OnlineServicesOSSAdapter/Source/Private/Online/StatsOSSAdapter.cpp
+++ b/Engine/Plugins/Online/Experimental/OnlineServicesOSSAdapter/Source/Private/Online/StatsOSSAdapter.cpp
@@ -208,7 +208,7 @@ TOnlineAsyncOpHandle<FBatchQueryStats> FStatsOSSAdapter::BatchQueryStats(FBatchQ
 		}
 
 		TArray<FUniqueNetIdRef> NetIds;
-		for (const FOnlineAccountIdHandle& TargetUserId : Op.GetParams().TargetUserIds)
+		for (const FAccountId& TargetUserId : Op.GetParams().TargetUserIds)
 		{
 			const FUniqueNetIdPtr UniqueNetId = Auth->GetUniqueNetId(TargetUserId);
 			if (!UniqueNetId)
diff --git a/Engine/Plugins/Online/Experimental/OnlineServicesOSSAdapter/Source/Private/Online/UserInfoOSSAdapter.cpp b/Engine/Plugins/Online/Experimental/OnlineServicesOSSAdapter/Source/Private/Online/UserInfoOSSAdapter.cpp
--- a/Engine/Plugins/Online/Experimental/OnlineServicesOSSAdapter/Source/Private/Online/UserInfoOSSAdapter.cpp
+++ b/Engine/Plugins/Online/Experimental/OnlineServicesOSSAdapter/Source/Private/Online/UserInfoOSSAdapter.cpp
@@ -48,7 +48,7 @@ TOnlineAsyncOpHandle<FQueryUserInfo> FUserInfoOSSAdapter::QueryUserInfo(FQueryUs
 
 		TArray<FUniqueNetIdRef> UserIds;
 		UserIds.Reserve(Params.UserIds.Num());
-		for(const FOnlineAccountIdHandle& TargetAccountIdHandle : Params.UserIds)
+		for(const FAccountId& TargetAccountIdHandle : Params.UserIds)
 		{
 			FUniqueNetIdPtr TargetUserNetId = AuthInterface->GetUniqueNetId(TargetAccountIdHandle);
 			if (!TargetUserNetId)
diff --git a/Engine/Plugins/Online/Experimental/OnlineServicesOSSAdapter/Source/Public/Online/AchievementsOSSAdapter.h b/Engine/Plugins/Online/Experimental/OnlineServicesOSSAdapter/Source/Public/Online/AchievementsOSSAdapter.h
--- a/Engine/Plugins/Online/Experimental/OnlineServicesOSSAdapter/Source/Public/Online/AchievementsOSSAdapter.h
+++ b/Engine/Plugins/Online/Experimental/OnlineServicesOSSAdapter/Source/Public/Online/AchievementsOSSAdapter.h
@@ -41,7 +41,7 @@ protected:
 	TOptional<FAchievementDefinitionMap> AchievementDefinitions;
 
 	using FAchievementStateMap = TMap<FString, FAchievementState>;
-	TMap<FOnlineAccountIdHandle, FAchievementStateMap> UserToAchievementStates;
+	TMap<FAccountId, FAchievementStateMap> UserToAchievementStates;
 };
 
 /* UE::Online */ }
diff --git a/Engine/Plugins/Online/Experimental/OnlineServicesOSSAdapter/Source/Public/Online/AuthOSSAdapter.h b/Engine/Plugins/Online/Experimental/OnlineServicesOSSAdapter/Source/Public/Online/AuthOSSAdapter.h
--- a/Engine/Plugins/Online/Experimental/OnlineServicesOSSAdapter/Source/Public/Online/AuthOSSAdapter.h
+++ b/Engine/Plugins/Online/Experimental/OnlineServicesOSSAdapter/Source/Public/Online/AuthOSSAdapter.h
@@ -21,7 +21,7 @@ struct FAuthHandleLoginStatusChangedImpl
 	struct Params
 	{
 		FPlatformUserId PlatformUserId;
-		FOnlineAccountIdHandle AccountId;
+		FAccountId AccountId;
 		ELoginStatus NewLoginStatus;
 	};
 
@@ -44,10 +44,10 @@ public:
 	virtual ~FAccountInfoRegistryOSSAdapter() = default;
 
 	TSharedPtr<FAccountInfoOSSAdapter> Find(FPlatformUserId PlatformUserId) const;
-	TSharedPtr<FAccountInfoOSSAdapter> Find(FOnlineAccountIdHandle AccountIdHandle) const;
+	TSharedPtr<FAccountInfoOSSAdapter> Find(FAccountId AccountIdHandle) const;
 
 	void Register(const TSharedRef<FAccountInfoOSSAdapter>&UserAuthData);
-	void Unregister(FOnlineAccountIdHandle AccountId);
+	void Unregister(FAccountId AccountId);
 };
 
 class FAuthOSSAdapter : public FAuthCommon
@@ -67,9 +67,9 @@ public:
 	virtual TOnlineAsyncOpHandle<FAuthQueryExternalServerAuthTicket> QueryExternalServerAuthTicket(FAuthQueryExternalServerAuthTicket::Params&& Params) override;
 	virtual TOnlineAsyncOpHandle<FAuthQueryExternalAuthToken> QueryExternalAuthToken(FAuthQueryExternalAuthToken::Params&& Params) override;
 
-	FUniqueNetIdPtr GetUniqueNetId(FOnlineAccountIdHandle AccountIdHandle) const;
-	FOnlineAccountIdHandle GetAccountIdHandle(const FUniqueNetIdRef& UniqueNetId) const;
-	int32 GetLocalUserNum(FOnlineAccountIdHandle AccountIdHandle) const;
+	FUniqueNetIdPtr GetUniqueNetId(FAccountId AccountIdHandle) const;
+	FAccountId GetAccountIdHandle(const FUniqueNetIdRef& UniqueNetId) const;
+	int32 GetLocalUserNum(FAccountId AccountIdHandle) const;
 
 protected:
 #if !UE_BUILD_SHIPPING
diff --git a/Engine/Plugins/Online/Experimental/OnlineServicesOSSAdapter/Source/Public/Online/PresenceOSSAdapter.h b/Engine/Plugins/Online/Experimental/OnlineServicesOSSAdapter/Source/Public/Online/PresenceOSSAdapter.h
--- a/Engine/Plugins/Online/Experimental/OnlineServicesOSSAdapter/Source/Public/Online/PresenceOSSAdapter.h
+++ b/Engine/Plugins/Online/Experimental/OnlineServicesOSSAdapter/Source/Public/Online/PresenceOSSAdapter.h
@@ -29,8 +29,8 @@ public:
 	virtual TOnlineAsyncOpHandle<FUpdatePresence> UpdatePresence(FUpdatePresence::Params&& Params) override;
 	virtual TOnlineAsyncOpHandle<FPartialUpdatePresence> PartialUpdatePresence(FPartialUpdatePresence::Params&& Params) override;
 
-	FUniqueNetIdRef GetUniqueNetId(FOnlineAccountIdHandle AccountIdHandle) const;
-	int32 GetLocalUserNum(FOnlineAccountIdHandle AccountIdHandle) const;
+	FUniqueNetIdRef GetUniqueNetId(FAccountId AccountIdHandle) const;
+	int32 GetLocalUserNum(FAccountId AccountIdHandle) const;
 
 protected:
 	const FAuthOSSAdapter* Auth;
diff --git a/Engine/Plugins/Online/Experimental/OnlineServicesOSSAdapter/Source/Public/Online/SessionsOSSAdapter.h b/Engine/Plugins/Online/Experimental/OnlineServicesOSSAdapter/Source/Public/Online/SessionsOSSAdapter.h
--- a/Engine/Plugins/Online/Experimental/OnlineServicesOSSAdapter/Source/Public/Online/SessionsOSSAdapter.h
+++ b/Engine/Plugins/Online/Experimental/OnlineServicesOSSAdapter/Source/Public/Online/SessionsOSSAdapter.h
@@ -105,7 +105,7 @@ protected:
 	TSet<FName> RegisteredSessionNames;
 
 	/** Cache of session search information to be used by ongoing session searches, indexed per user. Entries are removed upon search completion. */
-	TMap<FOnlineAccountIdHandle, TSharedRef<FOnlineSessionSearch>> PendingV1SessionSearchesPerUser;
+	TMap<FAccountId, TSharedRef<FOnlineSessionSearch>> PendingV1SessionSearchesPerUser;
 };
 
 /* UE::Online */ }
diff --git a/Engine/Plugins/Online/Experimental/OnlineServicesOSSAdapter/Source/Public/Online/UserFileOSSAdapter.h b/Engine/Plugins/Online/Experimental/OnlineServicesOSSAdapter/Source/Public/Online/UserFileOSSAdapter.h
--- a/Engine/Plugins/Online/Experimental/OnlineServicesOSSAdapter/Source/Public/Online/UserFileOSSAdapter.h
+++ b/Engine/Plugins/Online/Experimental/OnlineServicesOSSAdapter/Source/Public/Online/UserFileOSSAdapter.h
@@ -44,7 +44,7 @@ public:
 protected:
 	FUserFileOSSAdapterConfig Config;
 	IOnlineUserCloudPtr UserCloudInterface = nullptr;
-	TMap<FOnlineAccountIdHandle, TArray<FString>> EnumeratedFiles;
+	TMap<FAccountId, TArray<FString>> EnumeratedFiles;
 };
 
 /* UE::Online */ }
diff --git a/Engine/Source/Runtime/CoreOnline/Private/Online/CoreOnline.cpp b/Engine/Source/Runtime/CoreOnline/Private/Online/CoreOnline.cpp
--- a/Engine/Source/Runtime/CoreOnline/Private/Online/CoreOnline.cpp
+++ b/Engine/Source/Runtime/CoreOnline/Private/Online/CoreOnline.cpp
@@ -237,7 +237,7 @@ void FOnlineIdRegistryRegistry::UnregisterAccountIdRegistry(EOnlineServices Onli
 	}
 }
 
-FString FOnlineIdRegistryRegistry::ToLogString(const FOnlineAccountIdHandle& Handle) const
+FString FOnlineIdRegistryRegistry::ToLogString(const FAccountId& Handle) const
 {
 	FString Result;
 	if (IOnlineAccountIdRegistry* Registry = GetAccountIdRegistry(Handle.GetOnlineServicesType()))
@@ -251,7 +251,7 @@ FString FOnlineIdRegistryRegistry::ToLogString(const FOnlineAccountIdHandle& Han
 	return Result;
 }
 
-TArray<uint8> FOnlineIdRegistryRegistry::ToReplicationData(const FOnlineAccountIdHandle& Handle) const
+TArray<uint8> FOnlineIdRegistryRegistry::ToReplicationData(const FAccountId& Handle) const
 {
 	TArray<uint8> RepData;
 	if (IOnlineAccountIdRegistry* Registry = GetAccountIdRegistry(Handle.GetOnlineServicesType()))
@@ -265,9 +265,9 @@ TArray<uint8> FOnlineIdRegistryRegistry::ToReplicationData(const FOnlineAccountI
 	return RepData;
 }
 
-FOnlineAccountIdHandle FOnlineIdRegistryRegistry::ToAccountId(EOnlineServices Services, const TArray<uint8>& RepData) const
+FAccountId FOnlineIdRegistryRegistry::ToAccountId(EOnlineServices Services, const TArray<uint8>& RepData) const
 {
-	FOnlineAccountIdHandle Handle;
+	FAccountId Handle;
 	if (IOnlineAccountIdRegistry* Registry = GetAccountIdRegistry(Services))
 	{
 		Handle = Registry->FromReplicationData(RepData);
@@ -405,7 +405,7 @@ IOnlineSessionInviteIdRegistry* FOnlineIdRegistryRegistry::GetSessionInviteIdReg
 	return Registry;
 }
 
-FString ToLogString(const FOnlineAccountIdHandle& Id)
+FString ToLogString(const FAccountId& Id)
 {
 	return FString::Printf(TEXT("%s:%d (%s)"), LexToString(Id.GetOnlineServicesType()), Id.GetHandle(), *FOnlineIdRegistryRegistry::Get().ToLogString(Id));
 }
diff --git a/Engine/Source/Runtime/CoreOnline/Private/Online/CoreOnlinePrivate.cpp b/Engine/Source/Runtime/CoreOnline/Private/Online/CoreOnlinePrivate.cpp
--- a/Engine/Source/Runtime/CoreOnline/Private/Online/CoreOnlinePrivate.cpp
+++ b/Engine/Source/Runtime/CoreOnline/Private/Online/CoreOnlinePrivate.cpp
@@ -4,7 +4,7 @@
 
 namespace UE::Online {
 
-FString FOnlineForeignAccountIdRegistry::ToLogString(const FOnlineAccountIdHandle& Handle) const
+FString FOnlineForeignAccountIdRegistry::ToLogString(const FAccountId& Handle) const
 {
 	FString Result;
 	if (Handle.IsValid())
@@ -27,7 +27,7 @@ FString FOnlineForeignAccountIdRegistry::ToLogString(const FOnlineAccountIdHandl
 	return Result;
 }
 
-TArray<uint8> FOnlineForeignAccountIdRegistry::ToReplicationData(const FOnlineAccountIdHandle& Handle) const
+TArray<uint8> FOnlineForeignAccountIdRegistry::ToReplicationData(const FAccountId& Handle) const
 {
 	TArray<uint8> RepData;
 	if (Handle.IsValid())
@@ -46,20 +46,20 @@ TArray<uint8> FOnlineForeignAccountIdRegistry::ToReplicationData(const FOnlineAc
 	return RepData;
 }
 
-FOnlineAccountIdHandle FOnlineForeignAccountIdRegistry::FromReplicationData(EOnlineServices Services, const TArray<uint8>& RepData)
+FAccountId FOnlineForeignAccountIdRegistry::FromReplicationData(EOnlineServices Services, const TArray<uint8>& RepData)
 {
-	FOnlineAccountIdHandle Handle;
+	FAccountId Handle;
 	if (RepData.Num())
 	{
 		FRepData& RepDataForServices = OnlineServicesToRepData.FindOrAdd(Services);
-		if (FOnlineAccountIdHandle* Found = RepDataForServices.RepDataToHandle.Find(RepData))
+		if (FAccountId* Found = RepDataForServices.RepDataToHandle.Find(RepData))
 		{
 			Handle = *Found;
 		}
 		else
 		{
 			RepDataForServices.RepDataArray.Add(RepData);
-			Handle = FOnlineAccountIdHandle(Services, RepDataForServices.RepDataArray.Num());
+			Handle = FAccountId(Services, RepDataForServices.RepDataArray.Num());
 			RepDataForServices.RepDataToHandle.Emplace(RepData, Handle);
 		}
 	}
diff --git a/Engine/Source/Runtime/CoreOnline/Private/Online/CoreOnlinePrivate.h b/Engine/Source/Runtime/CoreOnline/Private/Online/CoreOnlinePrivate.h
--- a/Engine/Source/Runtime/CoreOnline/Private/Online/CoreOnlinePrivate.h
+++ b/Engine/Source/Runtime/CoreOnline/Private/Online/CoreOnlinePrivate.h
@@ -28,14 +28,14 @@ namespace UE::Online {
 class FOnlineForeignAccountIdRegistry
 {
 public:
-	FString ToLogString(const FOnlineAccountIdHandle& Handle) const;
-	TArray<uint8> ToReplicationData(const FOnlineAccountIdHandle& Handle) const;
-	FOnlineAccountIdHandle FromReplicationData(EOnlineServices Services, const TArray<uint8>& RepData);
+	FString ToLogString(const FAccountId& Handle) const;
+	TArray<uint8> ToReplicationData(const FAccountId& Handle) const;
+	FAccountId FromReplicationData(EOnlineServices Services, const TArray<uint8>& RepData);
 
 private:
 	struct FRepData
 	{
-		TMap<TArray<uint8>, FOnlineAccountIdHandle> RepDataToHandle;
+		TMap<TArray<uint8>, FAccountId> RepDataToHandle;
 		TArray<TArray<uint8>> RepDataArray;
 	};
 	TMap<EOnlineServices, FRepData> OnlineServicesToRepData;
diff --git a/Engine/Source/Runtime/CoreOnline/Private/Tests/CoreOnlineTests.cpp b/Engine/Source/Runtime/CoreOnline/Private/Tests/CoreOnlineTests.cpp
--- a/Engine/Source/Runtime/CoreOnline/Private/Tests/CoreOnlineTests.cpp
+++ b/Engine/Source/Runtime/CoreOnline/Private/Tests/CoreOnlineTests.cpp
@@ -14,7 +14,7 @@ IMPLEMENT_SIMPLE_AUTOMATION_TEST(FForeignIdRepTest,
 	const TArray<uint8> SourceRepData { 0, 1, 2, 3, 4, 5, 6, 7 };
 
 	FOnlineIdRegistryRegistry Registry;
-	const FOnlineAccountIdHandle Handle = Registry.ToAccountId(Services, SourceRepData);
+	const FAccountId Handle = Registry.ToAccountId(Services, SourceRepData);
 	const TArray<uint8> RepData = Registry.ToReplicationData(Handle);
 	UTEST_EQUAL(TEXT(""), RepData, SourceRepData);
 
diff --git a/Engine/Source/Runtime/CoreOnline/Public/Online/CoreOnline.h b/Engine/Source/Runtime/CoreOnline/Public/Online/CoreOnline.h
--- a/Engine/Source/Runtime/CoreOnline/Public/Online/CoreOnline.h
+++ b/Engine/Source/Runtime/CoreOnline/Public/Online/CoreOnline.h
@@ -306,14 +306,14 @@ private:
 	uint32 Value = uint32(EOnlineServices::Null) << 24;
 };
 
-using FOnlineAccountIdHandle = TOnlineIdHandle<OnlineIdHandleTags::FAccount>;
+using FAccountId = TOnlineIdHandle<OnlineIdHandleTags::FAccount>;
 using FOnlineLobbyIdHandle = TOnlineIdHandle<OnlineIdHandleTags::FLobby>;
 using FOnlineSessionIdHandle = TOnlineIdHandle<OnlineIdHandleTags::FSession>;
 using FOnlineSessionInviteIdHandle = TOnlineIdHandle<OnlineIdHandleTags::FSessionInvite>;
 using FOnlineVerifiedAuthTicketIdHandle = TOnlineIdHandle<OnlineIdHandleTags::FVerifiedAuthTicket>;
 using FOnlineVerifiedAuthSessionIdHandle = TOnlineIdHandle<OnlineIdHandleTags::FVerifiedAuthSession>;
 
-COREONLINE_API FString ToLogString(const FOnlineAccountIdHandle& Id);
+COREONLINE_API FString ToLogString(const FAccountId& Id);
 COREONLINE_API FString ToLogString(const FOnlineLobbyIdHandle& Id);
 COREONLINE_API FString ToLogString(const FOnlineSessionIdHandle& Id);
 COREONLINE_API FString ToLogString(const FOnlineSessionInviteIdHandle& Id);
@@ -374,9 +374,9 @@ public:
 	 */
 	COREONLINE_API void UnregisterAccountIdRegistry(EOnlineServices OnlineServices, int32 Priority = 0);
 
-	COREONLINE_API FString ToLogString(const FOnlineAccountIdHandle& Handle) const;
-	COREONLINE_API TArray<uint8> ToReplicationData(const FOnlineAccountIdHandle& Handle) const;
-	COREONLINE_API FOnlineAccountIdHandle ToAccountId(EOnlineServices Services, const TArray<uint8>& RepData) const;
+	COREONLINE_API FString ToLogString(const FAccountId& Handle) const;
+	COREONLINE_API TArray<uint8> ToReplicationData(const FAccountId& Handle) const;
+	COREONLINE_API FAccountId ToAccountId(EOnlineServices Services, const TArray<uint8>& RepData) const;
 
 	COREONLINE_API IOnlineAccountIdRegistry* GetAccountIdRegistry(EOnlineServices OnlineServices) const;
 
@@ -465,7 +465,7 @@ struct FUniqueNetIdWrapper
 {
 	//GENERATED_BODY()
 	
-	using FVariantType = TVariant<FUniqueNetIdPtr, UE::Online::FOnlineAccountIdHandle>;
+	using FVariantType = TVariant<FUniqueNetIdPtr, UE::Online::FAccountId>;
 
 	FUniqueNetIdWrapper() = default;
 	virtual ~FUniqueNetIdWrapper() = default;
@@ -487,9 +487,9 @@ struct FUniqueNetIdWrapper
 	{
 	}
 
-	FUniqueNetIdWrapper(const UE::Online::FOnlineAccountIdHandle& Handle)
+	FUniqueNetIdWrapper(const UE::Online::FAccountId& Handle)
 	{
-		Variant.Emplace<UE::Online::FOnlineAccountIdHandle>(Handle);
+		Variant.Emplace<UE::Online::FAccountId>(Handle);
 	}
 
 	// temporarily restored implicit conversion from FUniqueNetId
@@ -515,15 +515,15 @@ struct FUniqueNetIdWrapper
 
 	bool IsV2() const
 	{
-		return Variant.IsType<UE::Online::FOnlineAccountIdHandle>();
+		return Variant.IsType<UE::Online::FAccountId>();
 	}
 
-	UE::Online::FOnlineAccountIdHandle GetV2() const
+	UE::Online::FAccountId GetV2() const
 	{
-		UE::Online::FOnlineAccountIdHandle Result;
+		UE::Online::FAccountId Result;
 		if (ensure(IsV2()))
 		{
-			Result = Variant.Get<UE::Online::FOnlineAccountIdHandle>();
+			Result = Variant.Get<UE::Online::FAccountId>();
 		}
 		return Result;
 	}
@@ -562,7 +562,7 @@ struct FUniqueNetIdWrapper
 		}
 		else
 		{
-			const UE::Online::FOnlineAccountIdHandle& Handle = GetV2();
+			const UE::Online::FAccountId& Handle = GetV2();
 			return Handle.IsValid();
 		}
 	}
@@ -577,9 +577,9 @@ struct FUniqueNetIdWrapper
 		Variant.Emplace<FUniqueNetIdPtr>(InUniqueNetId);
 	}
 
-	virtual void SetAccountId(const UE::Online::FOnlineAccountIdHandle& Handle)
+	virtual void SetAccountId(const UE::Online::FAccountId& Handle)
 	{
-		Variant.Emplace<UE::Online::FOnlineAccountIdHandle>(Handle);
+		Variant.Emplace<UE::Online::FAccountId>(Handle);
 	}
 
 	/** @return unique id associated with this wrapper object */
diff --git a/Engine/Source/Runtime/Engine/Classes/GameFramework/OnlineReplStructs.h b/Engine/Source/Runtime/Engine/Classes/GameFramework/OnlineReplStructs.h
--- a/Engine/Source/Runtime/Engine/Classes/GameFramework/OnlineReplStructs.h
+++ b/Engine/Source/Runtime/Engine/Classes/GameFramework/OnlineReplStructs.h
@@ -62,7 +62,7 @@ struct FUniqueNetIdRepl : public FUniqueNetIdWrapper
 		FUniqueNetIdWrapper::SetUniqueNetId(UniqueNetId);
 	}
 
-	virtual void SetAccountId(const UE::Online::FOnlineAccountIdHandle& Handle) override
+	virtual void SetAccountId(const UE::Online::FAccountId& Handle) override
 	{
 		ReplicationBytes.Empty();
 		FUniqueNetIdWrapper::SetAccountId(Handle);
diff --git a/Engine/Source/Runtime/Engine/Private/OnlineReplStructs.cpp b/Engine/Source/Runtime/Engine/Private/OnlineReplStructs.cpp
--- a/Engine/Source/Runtime/Engine/Private/OnlineReplStructs.cpp
+++ b/Engine/Source/Runtime/Engine/Private/OnlineReplStructs.cpp
@@ -227,7 +227,7 @@ void FUniqueNetIdRepl::MakeReplicationDataV2()
 	EUniqueIdEncodingFlags EncodingFlags = EUniqueIdEncodingFlags::IsEncoded;
 	EncodingFlags |= static_cast<EUniqueIdEncodingFlags>(TypeHash_V2 << 3);
 
-	UE::Online::FOnlineAccountIdHandle Handle = GetV2();
+	UE::Online::FAccountId Handle = GetV2();
 	UE::Online::EOnlineServices OnlineServicesType = Handle.GetOnlineServicesType();
 	TArray<uint8> ReplicationData = UE::Online::FOnlineIdRegistryRegistry::Get().ToReplicationData(Handle);
 	check(!ReplicationData.IsEmpty());
@@ -445,7 +445,7 @@ void FUniqueNetIdRepl::NetSerializeLoadV2(FArchive& Ar, const EUniqueIdEncodingF
 	TArray<uint8> ReplicationData;
 	Ar << ReplicationData;
 
-	const UE::Online::FOnlineAccountIdHandle Handle = UE::Online::FOnlineIdRegistryRegistry::Get().ToAccountId(OnlineServicesType, ReplicationData);
+	const UE::Online::FAccountId Handle = UE::Online::FOnlineIdRegistryRegistry::Get().ToAccountId(OnlineServicesType, ReplicationData);
 	check(Handle.IsValid());
 	SetAccountId(Handle);
 }
