OSSv2: Rename Handle, UserId to AccountId when use FAccountId.

diff --git a/Engine/Plugins/Online/Experimental/OnlineServices/Source/OnlineServicesCommon/Private/Online/AchievementsCommon.cpp b/Engine/Plugins/Online/Experimental/OnlineServices/Source/OnlineServicesCommon/Private/Online/AchievementsCommon.cpp
--- a/Engine/Plugins/Online/Experimental/OnlineServices/Source/OnlineServicesCommon/Private/Online/AchievementsCommon.cpp
+++ b/Engine/Plugins/Online/Experimental/OnlineServices/Source/OnlineServicesCommon/Private/Online/AchievementsCommon.cpp
@@ -141,7 +141,7 @@ void FAchievementsCommon::UnlockAchievementsByStats(const FStatsUpdated& StatsUp
 	}
 
 	TArray<FString> StatNames;
-	TArray<FAccountId> UserIds;
+	TArray<FAccountId> AccountIds;
 	for (const FUserStats& UserStats : StatsUpdated.UpdateUsersStats)
 	{
 		for (const TPair<FString, FStatValue>& StatPair : UserStats.Stats)
@@ -158,17 +158,17 @@ void FAchievementsCommon::UnlockAchievementsByStats(const FStatsUpdated& StatsUp
 			}
 		}
 
-		UserIds.AddUnique(UserStats.UserId);
+		AccountIds.AddUnique(UserStats.AccountId);
 	}
 
-	if (StatNames.IsEmpty() || UserIds.IsEmpty())
+	if (StatNames.IsEmpty() || AccountIds.IsEmpty())
 	{
 		return;
 	}
 
 	FBatchQueryStats::Params BatchQueryStatsParam;
-	BatchQueryStatsParam.LocalUserId = StatsUpdated.LocalUserId;
-	BatchQueryStatsParam.TargetUserIds = MoveTemp(UserIds);
+	BatchQueryStatsParam.LocalAccountId = StatsUpdated.LocalAccountId;
+	BatchQueryStatsParam.TargetAccountIds = MoveTemp(AccountIds);
 	BatchQueryStatsParam.StatNames = MoveTemp(StatNames);
 
 	Services.Get<FStatsCommon>()->BatchQueryStats(MoveTemp(BatchQueryStatsParam))
@@ -182,13 +182,13 @@ void FAchievementsCommon::UnlockAchievementsByStats(const FStatsUpdated& StatsUp
 				TArray<FString> AchievementsToUnlock;
 				for (const TPair<FString, FStatValue>& StatPair : UserStats.Stats)
 				{
-					ExecuteUnlockRulesRelatedToStat(UserStats.UserId, StatPair.Key, UserStats.Stats, AchievementsToUnlock);
+					ExecuteUnlockRulesRelatedToStat(UserStats.AccountId, StatPair.Key, UserStats.Stats, AchievementsToUnlock);
 				}
 
 				if (!AchievementsToUnlock.IsEmpty())
 				{
 					FUnlockAchievements::Params UnlockAchievementsParams;
-					UnlockAchievementsParams.LocalUserId = UserStats.UserId;
+					UnlockAchievementsParams.LocalAccountId = UserStats.AccountId;
 					UnlockAchievementsParams.AchievementIds = MoveTemp(AchievementsToUnlock);
 					UnlockAchievements(MoveTemp(UnlockAchievementsParams));
 				}
@@ -197,12 +197,12 @@ void FAchievementsCommon::UnlockAchievementsByStats(const FStatsUpdated& StatsUp
 	});
 }
 
-void FAchievementsCommon::ExecuteUnlockRulesRelatedToStat(const FAccountId& UserId, const FString& StatName, const TMap<FString, FStatValue>& Stats, TArray<FString>& OutAchievementsToUnlock)
+void FAchievementsCommon::ExecuteUnlockRulesRelatedToStat(const FAccountId& AccountId, const FString& StatName, const TMap<FString, FStatValue>& Stats, TArray<FString>& OutAchievementsToUnlock)
 {
 	for (const FAchievementUnlockRule& AchievementUnlockRule : AchievementUnlockRules)
 	{
 		if (AchievementUnlockRule.ContainsStat(StatName)
-			&& !IsUnlocked(UserId, AchievementUnlockRule.AchievementToUnlock) 
+			&& !IsUnlocked(AccountId, AchievementUnlockRule.AchievementToUnlock) 
 			&& MeetUnlockCondition(AchievementUnlockRule, Stats))
 		{
 			OutAchievementsToUnlock.AddUnique(AchievementUnlockRule.AchievementToUnlock);
@@ -251,10 +251,10 @@ bool FAchievementsCommon::MeetUnlockCondition(FAchievementUnlockRule Achievement
 	return true;
 }
 
-bool FAchievementsCommon::IsUnlocked(const FAccountId& UserId, const FString& AchievementName) const
+bool FAchievementsCommon::IsUnlocked(const FAccountId& AccountId, const FString& AchievementName) const
 {
 	FGetAchievementState::Params Params;
-	Params.LocalUserId = UserId;
+	Params.LocalAccountId = AccountId;
 	Params.AchievementId = AchievementName;
 	TOnlineResult<FGetAchievementState> Result = GetAchievementState(MoveTemp(Params));
 	if (Result.IsOk())
diff --git a/Engine/Plugins/Online/Experimental/OnlineServices/Source/OnlineServicesCommon/Private/Online/AuthCommon.cpp b/Engine/Plugins/Online/Experimental/OnlineServices/Source/OnlineServicesCommon/Private/Online/AuthCommon.cpp
--- a/Engine/Plugins/Online/Experimental/OnlineServices/Source/OnlineServicesCommon/Private/Online/AuthCommon.cpp
+++ b/Engine/Plugins/Online/Experimental/OnlineServices/Source/OnlineServicesCommon/Private/Online/AuthCommon.cpp
@@ -17,10 +17,10 @@ TSharedPtr<FAccountInfo> FAccountInfoRegistry::Find(FPlatformUserId PlatformUser
 	return FoundPtr ? TSharedPtr<FAccountInfo>(*FoundPtr) : TSharedPtr<FAccountInfo>();
 }
 
-TSharedPtr<FAccountInfo> FAccountInfoRegistry::Find(FAccountId AccountIdHandle) const
+TSharedPtr<FAccountInfo> FAccountInfoRegistry::Find(FAccountId AccountId) const
 {
 	FReadScopeLock Lock(IndexLock);
-	const TSharedRef<FAccountInfo>* FoundPtr = AuthDataByOnlineAccountIdHandle.Find(AccountIdHandle);
+	const TSharedRef<FAccountInfo>* FoundPtr = AuthDataByOnlineAccountIdHandle.Find(AccountId);
 	return FoundPtr ? TSharedPtr<FAccountInfo>(*FoundPtr) : TSharedPtr<FAccountInfo>();
 }
 
@@ -162,7 +162,7 @@ TOnlineResult<FAuthGetLocalOnlineUserByPlatformUserId> FAuthCommon::GetLocalOnli
 
 TOnlineResult<FAuthGetLocalOnlineUserByOnlineAccountId> FAuthCommon::GetLocalOnlineUserByOnlineAccountId(FAuthGetLocalOnlineUserByOnlineAccountId::Params&& Params) const
 {
-	TSharedPtr<FAccountInfo> AccountInfo = GetAccountInfoRegistry().Find(Params.LocalUserId);
+	TSharedPtr<FAccountInfo> AccountInfo = GetAccountInfoRegistry().Find(Params.LocalAccountId);
 	return (AccountInfo && IsOnlineStatus(AccountInfo->LoginStatus)) ?
 		TOnlineResult<FAuthGetLocalOnlineUserByOnlineAccountId>(FAuthGetLocalOnlineUserByOnlineAccountId::Result{ AccountInfo.ToSharedRef() }) :
 		TOnlineResult<FAuthGetLocalOnlineUserByOnlineAccountId>(Errors::InvalidUser());
diff --git a/Engine/Plugins/Online/Experimental/OnlineServices/Source/OnlineServicesCommon/Private/Online/LeaderboardsCommon.cpp b/Engine/Plugins/Online/Experimental/OnlineServices/Source/OnlineServicesCommon/Private/Online/LeaderboardsCommon.cpp
--- a/Engine/Plugins/Online/Experimental/OnlineServices/Source/OnlineServicesCommon/Private/Online/LeaderboardsCommon.cpp
+++ b/Engine/Plugins/Online/Experimental/OnlineServices/Source/OnlineServicesCommon/Private/Online/LeaderboardsCommon.cpp
@@ -166,7 +166,7 @@ void FLeaderboardsCommon::WriteLeaderboardsByStats(const FStatsUpdated& StatsUpd
 			if (LeaderboardDefinitions.Contains(StatPair.Key))
 			{
 				FWriteLeaderboardScores::Params WriteLeaderboardScoresParam;
-				WriteLeaderboardScoresParam.LocalUserId = UserStats.UserId;
+				WriteLeaderboardScoresParam.LocalAccountId = UserStats.AccountId;
 				WriteLeaderboardScoresParam.BoardName = StatPair.Key;
 				WriteLeaderboardScoresParam.Score = StatPair.Value.GetInt64();
 				WriteLeaderboardScores(MoveTemp(WriteLeaderboardScoresParam));
diff --git a/Engine/Plugins/Online/Experimental/OnlineServices/Source/OnlineServicesCommon/Private/Online/LobbiesCommonTests.cpp b/Engine/Plugins/Online/Experimental/OnlineServices/Source/OnlineServicesCommon/Private/Online/LobbiesCommonTests.cpp
--- a/Engine/Plugins/Online/Experimental/OnlineServices/Source/OnlineServicesCommon/Private/Online/LobbiesCommonTests.cpp
+++ b/Engine/Plugins/Online/Experimental/OnlineServices/Source/OnlineServicesCommon/Private/Online/LobbiesCommonTests.cpp
@@ -1145,7 +1145,7 @@ TFuture<TDefaultErrorResultInternal<FOnlineLobbyIdHandle>> AwaitInvitation(
 
 		if (!AwaitState->bSignaled)
 		{
-			if (Notification.Lobby->LobbyId == LobbyId && Notification.LocalUserId == TargetAccountId)
+			if (Notification.Lobby->LobbyId == LobbyId && Notification.LocalAccountId == TargetAccountId)
 			{
 				FTSTicker::GetCoreTicker().RemoveTicker(AwaitState->OnAwaitExpiredHandle);
 				AwaitState->bSignaled = true;
@@ -1292,7 +1292,7 @@ TFuture<TOnlineResult<FFunctionalTestLogoutUser>> FunctionalTestLogoutUser(FFunc
 	TFuture<TOnlineResult<FFunctionalTestLogoutUser>> Future = Promise->GetFuture();
 
 	FAuthLogout::Params LogoutParams;
-	LogoutParams.LocalUserId = AccountInfo->AccountId;
+	LogoutParams.LocalAccountId = AccountInfo->AccountId;
 
 	Params.AuthInterface->Logout(MoveTemp(LogoutParams))
 	.OnComplete([Promise, PlatformUserId = Params.PlatformUserId](const TOnlineResult<FAuthLogout>& LogoutResult)
@@ -1468,7 +1468,7 @@ TOnlineAsyncOpHandle<FFunctionalTestLobbies> RunLobbyFunctionalTest(IAuth& AuthI
 
 		// Create a lobby
 		FCreateLobby::Params Params;
-		Params.LocalUserId = User1Info.AccountInfo->AccountId;
+		Params.LocalAccountId = User1Info.AccountInfo->AccountId;
 		Params.LocalName = TEXT("test");
 		Params.SchemaName = TEXT("test");
 		Params.MaxMembers = 2;
@@ -1522,7 +1522,7 @@ TOnlineAsyncOpHandle<FFunctionalTestLobbies> RunLobbyFunctionalTest(IAuth& AuthI
 		const FCreateLobby::Result& CreateLobbyResult = GetOpDataChecked<FCreateLobby::Result>(InAsyncOp, CreateLobbyKeyName);
 
 		FModifyLobbyAttributes::Params Params;
-		Params.LocalUserId = User1Info.AccountInfo->AccountId;
+		Params.LocalAccountId = User1Info.AccountInfo->AccountId;
 		Params.LobbyId = CreateLobbyResult.Lobby->LobbyId;
 		Params.MutatedAttributes = {{TEXT("test_attribute_key"), TEXT("mutated_test_attribute_value")}};
 		return Params;
@@ -1555,7 +1555,7 @@ TOnlineAsyncOpHandle<FFunctionalTestLobbies> RunLobbyFunctionalTest(IAuth& AuthI
 		const FCreateLobby::Result& CreateLobbyResult = GetOpDataChecked<FCreateLobby::Result>(InAsyncOp, CreateLobbyKeyName);
 
 		FModifyLobbyMemberAttributes::Params Params;
-		Params.LocalUserId = User1Info.AccountInfo->AccountId;
+		Params.LocalAccountId = User1Info.AccountInfo->AccountId;
 		Params.LobbyId = CreateLobbyResult.Lobby->LobbyId;
 		Params.MutatedAttributes = {{TEXT("test_attribute_key"), TEXT("mutated_test_attribute_value")}};
 		return Params;
@@ -1597,7 +1597,7 @@ TOnlineAsyncOpHandle<FFunctionalTestLobbies> RunLobbyFunctionalTest(IAuth& AuthI
 		const FCreateLobby::Result& CreateLobbyResult = GetOpDataChecked<FCreateLobby::Result>(InAsyncOp, CreateLobbyKeyName);
 
 		FLeaveLobby::Params Params;
-		Params.LocalUserId = User1Info.AccountInfo->AccountId;
+		Params.LocalAccountId = User1Info.AccountInfo->AccountId;
 		Params.LobbyId = CreateLobbyResult.Lobby->LobbyId;
 		return Params;
 	})
@@ -1631,7 +1631,7 @@ TOnlineAsyncOpHandle<FFunctionalTestLobbies> RunLobbyFunctionalTest(IAuth& AuthI
 		const Private::FFunctionalTestLoginUser::Result& User1Info = GetOpDataChecked<Private::FFunctionalTestLoginUser::Result>(InAsyncOp, User1KeyName);
 
 		FCreateLobby::Params Params;
-		Params.LocalUserId = User1Info.AccountInfo->AccountId;
+		Params.LocalAccountId = User1Info.AccountInfo->AccountId;
 		Params.LocalName = TEXT("test");
 		Params.SchemaName = TEXT("test");
 		Params.MaxMembers = 2;
@@ -1662,9 +1662,9 @@ TOnlineAsyncOpHandle<FFunctionalTestLobbies> RunLobbyFunctionalTest(IAuth& AuthI
 		const FCreateLobby::Result& CreateLobbyResult = GetOpDataChecked<FCreateLobby::Result>(InAsyncOp, CreateLobbyKeyName);
 
 		FInviteLobbyMember::Params Params;
-		Params.LocalUserId = User1Info.AccountInfo->AccountId;
+		Params.LocalAccountId = User1Info.AccountInfo->AccountId;
 		Params.LobbyId = CreateLobbyResult.Lobby->LobbyId;
-		Params.TargetUserId = User2Info.AccountInfo->AccountId;
+		Params.TargetAccountId = User2Info.AccountInfo->AccountId;
 		return Params;
 	})
 	.Then(Private::ConsumeOperationStepResult<FInviteLobbyMember>(*Op, Private::FBindOperation<FInviteLobbyMember>(&LobbiesCommon, &FLobbiesCommon::InviteLobbyMember)))
@@ -1690,7 +1690,7 @@ TOnlineAsyncOpHandle<FFunctionalTestLobbies> RunLobbyFunctionalTest(IAuth& AuthI
 		const FCreateLobby::Result& CreateLobbyResult = GetOpDataChecked<FCreateLobby::Result>(InAsyncOp, CreateLobbyKeyName);
 
 		FJoinLobby::Params Params;
-		Params.LocalUserId = User2Info.AccountInfo->AccountId;
+		Params.LocalAccountId = User2Info.AccountInfo->AccountId;
 		Params.LocalName = TEXT("test");
 		Params.LobbyId = CreateLobbyResult.Lobby->LobbyId;
 		Params.LocalUsers.Emplace(FJoinLobbyLocalUserData{User2Info.AccountInfo->AccountId, {}});
@@ -1715,7 +1715,7 @@ TOnlineAsyncOpHandle<FFunctionalTestLobbies> RunLobbyFunctionalTest(IAuth& AuthI
 		const FCreateLobby::Result& CreateLobbyResult = GetOpDataChecked<FCreateLobby::Result>(InAsyncOp, CreateLobbyKeyName);
 
 		FLeaveLobby::Params Params;
-		Params.LocalUserId = User2Info.AccountInfo->AccountId;
+		Params.LocalAccountId = User2Info.AccountInfo->AccountId;
 		Params.LobbyId = CreateLobbyResult.Lobby->LobbyId;
 		return Params;
 	})
@@ -1738,7 +1738,7 @@ TOnlineAsyncOpHandle<FFunctionalTestLobbies> RunLobbyFunctionalTest(IAuth& AuthI
 		const FCreateLobby::Result& CreateLobbyResult = GetOpDataChecked<FCreateLobby::Result>(InAsyncOp, CreateLobbyKeyName);
 
 		FLeaveLobby::Params Params;
-		Params.LocalUserId = User1Info.AccountInfo->AccountId;
+		Params.LocalAccountId = User1Info.AccountInfo->AccountId;
 		Params.LobbyId = CreateLobbyResult.Lobby->LobbyId;
 		return Params;
 	})
@@ -1780,7 +1780,7 @@ TOnlineAsyncOpHandle<FFunctionalTestLobbies> RunLobbyFunctionalTest(IAuth& AuthI
 		InAsyncOp.Data.Set(SearchKeyName, TSharedRef<FSearchParams>(SearchParams));
 
 		FCreateLobby::Params Params;
-		Params.LocalUserId = User1Info.AccountInfo->AccountId;
+		Params.LocalAccountId = User1Info.AccountInfo->AccountId;
 		Params.LocalName = TEXT("test");
 		Params.SchemaName = TEXT("test");
 		Params.MaxMembers = 2;
@@ -1815,7 +1815,7 @@ TOnlineAsyncOpHandle<FFunctionalTestLobbies> RunLobbyFunctionalTest(IAuth& AuthI
 
 		// Search for lobby from create. Searching by attribute will also limit results by bucket id.
 		FFindLobbies::Params Params;
-		Params.LocalUserId = User2Info.AccountInfo->AccountId;
+		Params.LocalAccountId = User2Info.AccountInfo->AccountId;
 		Params.Filters = {{TEXT("LobbyCreateTime"), ELobbyComparisonOp::Equals, SearchParams->LobbyCreateTime}};
 		return Params;
 	})
@@ -1836,7 +1836,7 @@ TOnlineAsyncOpHandle<FFunctionalTestLobbies> RunLobbyFunctionalTest(IAuth& AuthI
 		}
 
 		FJoinLobby::Params Params;
-		Params.LocalUserId = User2Info.AccountInfo->AccountId;
+		Params.LocalAccountId = User2Info.AccountInfo->AccountId;
 		Params.LocalName = TEXT("test");
 		Params.LobbyId = FindResults.Lobbies[0]->LobbyId;
 		Params.LocalUsers.Emplace(FJoinLobbyLocalUserData{User2Info.AccountInfo->AccountId, {}});
@@ -1861,7 +1861,7 @@ TOnlineAsyncOpHandle<FFunctionalTestLobbies> RunLobbyFunctionalTest(IAuth& AuthI
 		const FCreateLobby::Result& CreateLobbyResult = GetOpDataChecked<FCreateLobby::Result>(InAsyncOp, CreateLobbyKeyName);
 
 		FLeaveLobby::Params Params;
-		Params.LocalUserId = User2Info.AccountInfo->AccountId;
+		Params.LocalAccountId = User2Info.AccountInfo->AccountId;
 		Params.LobbyId = CreateLobbyResult.Lobby->LobbyId;
 		return Params;
 	})
@@ -1884,7 +1884,7 @@ TOnlineAsyncOpHandle<FFunctionalTestLobbies> RunLobbyFunctionalTest(IAuth& AuthI
 		const FCreateLobby::Result& CreateLobbyResult = GetOpDataChecked<FCreateLobby::Result>(InAsyncOp, CreateLobbyKeyName);
 
 		FLeaveLobby::Params Params;
-		Params.LocalUserId = User1Info.AccountInfo->AccountId;
+		Params.LocalAccountId = User1Info.AccountInfo->AccountId;
 		Params.LobbyId = CreateLobbyResult.Lobby->LobbyId;
 		return Params;
 	})
diff --git a/Engine/Plugins/Online/Experimental/OnlineServices/Source/OnlineServicesCommon/Private/Online/OnlineSyntaxTests.cpp b/Engine/Plugins/Online/Experimental/OnlineServices/Source/OnlineServicesCommon/Private/Online/OnlineSyntaxTests.cpp
--- a/Engine/Plugins/Online/Experimental/OnlineServices/Source/OnlineServicesCommon/Private/Online/OnlineSyntaxTests.cpp
+++ b/Engine/Plugins/Online/Experimental/OnlineServices/Source/OnlineServicesCommon/Private/Online/OnlineSyntaxTests.cpp
@@ -47,7 +47,7 @@ struct FTestUserOp
 
 	struct Params
 	{
-		FAccountId LocalUserId;
+		FAccountId LocalAccountId;
 	};
 
 	struct Result
@@ -74,7 +74,7 @@ struct FJoinableTestUserOp
 
 	struct Params
 	{
-		FAccountId LocalUserId;
+		FAccountId LocalAccountId;
 	};
 
 	struct Result
@@ -109,7 +109,7 @@ struct FMergeableTestUserOp
 
 	struct Params
 	{
-		FAccountId LocalUserId;
+		FAccountId LocalAccountId;
 		FTestMutation Mutations;
 	};
 
@@ -126,14 +126,14 @@ BEGIN_ONLINE_STRUCT_META(Test::FTestOp::Params)
 END_ONLINE_STRUCT_META()
 
 BEGIN_ONLINE_STRUCT_META(Test::FTestUserOp::Params)
-	ONLINE_STRUCT_FIELD(Test::FTestUserOp::Params, LocalUserId)
+	ONLINE_STRUCT_FIELD(Test::FTestUserOp::Params, LocalAccountId)
 END_ONLINE_STRUCT_META()
 
 BEGIN_ONLINE_STRUCT_META(Test::FJoinableTestOp::Params)
 END_ONLINE_STRUCT_META()
 
 BEGIN_ONLINE_STRUCT_META(Test::FJoinableTestUserOp::Params)
-	ONLINE_STRUCT_FIELD(Test::FJoinableTestUserOp::Params, LocalUserId)
+	ONLINE_STRUCT_FIELD(Test::FJoinableTestUserOp::Params, LocalAccountId)
 END_ONLINE_STRUCT_META()
 
 BEGIN_ONLINE_STRUCT_META(Test::FMergeableTestOp::Params)
@@ -142,7 +142,7 @@ END_ONLINE_STRUCT_META()
 
 BEGIN_ONLINE_STRUCT_META(Test::FMergeableTestUserOp::Params)
 	ONLINE_STRUCT_FIELD(Test::FMergeableTestUserOp::Params, Mutations),
-	ONLINE_STRUCT_FIELD(Test::FMergeableTestUserOp::Params, LocalUserId)
+	ONLINE_STRUCT_FIELD(Test::FMergeableTestUserOp::Params, LocalAccountId)
 END_ONLINE_STRUCT_META()
 
 /* Meta */ }
@@ -150,8 +150,8 @@ END_ONLINE_STRUCT_META()
 // confirm that a couple concepts are working as expected
 static_assert(TModels<UE::Online::Meta::CSuperDefined, Test::FTestInterface>::Value);
 static_assert(!TModels<UE::Online::Meta::CSuperDefined, Test::ITestInterface>::Value);
-static_assert(!TModels<UE::Online::CLocalUserIdDefined, Test::FTestOp::Params>::Value);
-static_assert(TModels<UE::Online::CLocalUserIdDefined, Test::FTestUserOp::Params>::Value);
+static_assert(!TModels<UE::Online::CLocalAccountIdDefined, Test::FTestOp::Params>::Value);
+static_assert(TModels<UE::Online::CLocalAccountIdDefined, Test::FTestUserOp::Params>::Value);
 
 /* UE::Online */ }
 
diff --git a/Engine/Plugins/Online/Experimental/OnlineServices/Source/OnlineServicesCommon/Private/Online/PresenceCommon.cpp b/Engine/Plugins/Online/Experimental/OnlineServices/Source/OnlineServicesCommon/Private/Online/PresenceCommon.cpp
--- a/Engine/Plugins/Online/Experimental/OnlineServices/Source/OnlineServicesCommon/Private/Online/PresenceCommon.cpp
+++ b/Engine/Plugins/Online/Experimental/OnlineServices/Source/OnlineServicesCommon/Private/Online/PresenceCommon.cpp
@@ -36,16 +36,16 @@ public:
 		, Presence(InPresence)
 		, ErrorState(Errors::Success())
 	{
-		NumPresencesRemaining = Op->GetParams().TargetUserIds.Num();
+		NumPresencesRemaining = Op->GetParams().TargetAccountIds.Num();
 	}
 
 	TFuture<FBatchQueryPresenceHelperResult> GetPromise()
 	{
-		for(FAccountId TargetUserId : Op->GetParams().TargetUserIds)
+		for(FAccountId TargetAccountId : Op->GetParams().TargetAccountIds)
 		{
 			FQueryPresence::Params Params;
-			Params.LocalUserId = Op->GetParams().LocalUserId;
-			Params.TargetUserId = TargetUserId;
+			Params.LocalAccountId = Op->GetParams().LocalAccountId;
+			Params.TargetAccountId = TargetAccountId;
 			Params.bListenToChanges = Op->GetParams().bListenToChanges;
 
 			Presence->QueryPresence(MoveTemp(Params)).OnComplete([this](const TOnlineResult<FQueryPresence> NewOp) mutable
@@ -122,7 +122,7 @@ TOnlineAsyncOpHandle<FUpdatePresence> FPresenceCommon::UpdatePresence(FUpdatePre
 
 TOnlineAsyncOpHandle<FPartialUpdatePresence> FPresenceCommon::PartialUpdatePresence(FPartialUpdatePresence::Params&& Params)
 {
-	TOnlineResult<FGetCachedPresence> CachedPresence = GetCachedPresence({Params.LocalUserId, Params.LocalUserId});
+	TOnlineResult<FGetCachedPresence> CachedPresence = GetCachedPresence({Params.LocalAccountId, Params.LocalAccountId});
 	if (CachedPresence.IsOk())
 	{
 		TSharedRef<FUserPresence> NewPresence = MakeShared<FUserPresence>();
@@ -155,7 +155,7 @@ TOnlineAsyncOpHandle<FPartialUpdatePresence> FPresenceCommon::PartialUpdatePrese
 		}
 
 		TOnlineAsyncOpRef<FPartialUpdatePresence> Op = GetOp<FPartialUpdatePresence>(MoveTemp(Params));
-		TOnlineAsyncOpHandle<FUpdatePresence> UpdatePresenceResult = UpdatePresence({Op->GetParams().LocalUserId, NewPresence});
+		TOnlineAsyncOpHandle<FUpdatePresence> UpdatePresenceResult = UpdatePresence({Op->GetParams().LocalAccountId, NewPresence});
 		UpdatePresenceResult.OnComplete([this, Op](const TOnlineResult<FUpdatePresence> Result)
 		{
 			if (Result.IsOk())
diff --git a/Engine/Plugins/Online/Experimental/OnlineServices/Source/OnlineServicesCommon/Private/Online/SessionsCommon.cpp b/Engine/Plugins/Online/Experimental/OnlineServices/Source/OnlineServicesCommon/Private/Online/SessionsCommon.cpp
--- a/Engine/Plugins/Online/Experimental/OnlineServices/Source/OnlineServicesCommon/Private/Online/SessionsCommon.cpp
+++ b/Engine/Plugins/Online/Experimental/OnlineServices/Source/OnlineServicesCommon/Private/Online/SessionsCommon.cpp
@@ -68,12 +68,12 @@ namespace UE::Online {
 		if (GetMutableSessionByIdResult.IsOk())
 		{
 			return TOnlineResult<FGetSessionById>({ GetMutableSessionByIdResult.GetOkValue().Session });
-		}
+			}
 		else
 		{
 			return TOnlineResult<FGetSessionById>(GetMutableSessionByIdResult.GetErrorValue());
+			}
 		}
-	}
 
 	TOnlineAsyncOpHandle<FCreateSession> FSessionsCommon::CreateSession(FCreateSession::Params&& Params)
 	{
@@ -188,7 +188,7 @@ namespace UE::Online {
 	{
 		TArray<TSharedRef<const FSessionInvite>> SessionInvites;
 
-		if (const TMap<FOnlineSessionInviteIdHandle, TSharedRef<FSessionInvite>>* UserMap = SessionInvitesUserMap.Find(Params.LocalUserId))
+		if (const TMap<FOnlineSessionInviteIdHandle, TSharedRef<FSessionInvite>>* UserMap = SessionInvitesUserMap.Find(Params.LocalAccountId))
 		{
 			for (const TPair<FOnlineSessionInviteIdHandle, TSharedRef<FSessionInvite>>& Entry : *UserMap)
 			{
@@ -413,7 +413,7 @@ namespace UE::Online {
 			}
 			else if (Params.bRegisterPlayers) // If it wasn't and the parameters dictate it, we'll register them with a reserved slot
 			{
-				FRegisterPlayers::Params RegisterPlayerParams = { Params.LocalUserId, Params.SessionName, { SessionMemberId }, true };
+				FRegisterPlayers::Params RegisterPlayerParams = { Params.LocalAccountId, Params.SessionName, { SessionMemberId }, true };
 
 				TOnlineResult<FRegisterPlayers> Result = RegisterPlayersImpl(RegisterPlayerParams);
 				if (Result.IsError())
@@ -456,7 +456,7 @@ namespace UE::Online {
 		{
 			if (Params.bUnregisterPlayers) // If the parameters dictate it, well unregister the player as well
 			{
-				FUnregisterPlayers::Params UnregisterPlayersParams = { Params.LocalUserId, Params.SessionName, { SessionMemberId }, false };
+				FUnregisterPlayers::Params UnregisterPlayersParams = { Params.LocalAccountId, Params.SessionName, { SessionMemberId }, false };
 
 				TOnlineResult<FUnregisterPlayers> Result = UnregisterPlayersImpl(UnregisterPlayersParams);
 				if (Result.IsError())
@@ -595,16 +595,16 @@ namespace UE::Online {
 			return Errors::InvalidParams();
 		}
 
-		if (!Params.LocalUserId.IsValid())
+		if (!Params.LocalAccountId.IsValid())
 		{
-			UE_LOG(LogTemp, Warning, TEXT("[FSessionsCommon::CheckCreateSessionParams] Could not create session with name [%s]. LocalUserId [%s] not valid"), *Params.SessionName.ToString(), *ToLogString(Params.LocalUserId));
+			UE_LOG(LogTemp, Warning, TEXT("[FSessionsCommon::CheckCreateSessionParams] Could not create session with name [%s]. LocalAccountId [%s] not valid"), *Params.SessionName.ToString(), *ToLogString(Params.LocalAccountId));
 
 			return Errors::InvalidParams();
 		}
 
-		if (!Params.LocalUsers.Contains(Params.LocalUserId))
+		if (!Params.LocalAccounts.Contains(Params.LocalAccountId))
 		{
-			UE_LOG(LogTemp, Warning, TEXT("[FSessionsCommon::CheckCreateSessionParams] Could not create session with name [%s]. Missing Session Member info for user [%s]"), *Params.SessionName.ToString(), *ToLogString(Params.LocalUserId));
+			UE_LOG(LogTemp, Warning, TEXT("[FSessionsCommon::CheckCreateSessionParams] Could not create session with name [%s]. Missing Session Member info for user [%s]"), *Params.SessionName.ToString(), *ToLogString(Params.LocalAccountId));
 
 			return Errors::InvalidParams();
 		}
@@ -633,31 +633,31 @@ namespace UE::Online {
 		{
 			for (const TPair<FName, FOnlineSessionIdHandle>& Entry : LocalSessionsByName)
 			{
-				TOnlineResult<FGetSessionById> GetSessionByIdResult = GetSessionById({ Params.LocalUserId, Entry.Value });
+				TOnlineResult<FGetSessionById> GetSessionByIdResult = GetSessionById({ Params.LocalAccountId, Entry.Value });
 				if (GetSessionByIdResult.IsOk())
 				{
 					if (GetSessionByIdResult.GetOkValue().Session->SessionSettings.bPresenceEnabled)
-					{
-						UE_LOG(LogTemp, Warning, TEXT("[FSessionsCommon::CheckCreateSessionState] Could not create session with bPresenceEnabled set to true when another already exists [%s]."), *Entry.Key.ToString());
+				{
+					UE_LOG(LogTemp, Warning, TEXT("[FSessionsCommon::CheckCreateSessionState] Could not create session with bPresenceEnabled set to true when another already exists [%s]."), *Entry.Key.ToString());
 
-						return Errors::InvalidState();
-					}
+					return Errors::InvalidState();
 				}
 			}
 		}
+		}
 
 		// User login check for all local users
 		IAuthPtr Auth = Services.GetAuthInterface();
 
-		TArray<FAccountId> LocalUserIds;
-		LocalUserIds.Reserve(Params.LocalUsers.Num());
-		Params.LocalUsers.GenerateKeyArray(LocalUserIds);
+		TArray<FAccountId> LocalAccountIds;
+		LocalAccountIds.Reserve(Params.LocalAccounts.Num());
+		Params.LocalAccounts.GenerateKeyArray(LocalAccountIds);
 
-		for (const FAccountId& LocalUserId : LocalUserIds)
+		for (const FAccountId& LocalAccountId : LocalAccountIds)
 		{
-			if (!Auth->IsLoggedIn(LocalUserId))
+			if (!Auth->IsLoggedIn(LocalAccountId))
 			{
-				UE_LOG(LogTemp, Warning, TEXT("[FSessionsCommon::CheckCreateSessionState] Could not create session with user [%s] not logged in"), *ToLogString(LocalUserId));
+				UE_LOG(LogTemp, Warning, TEXT("[FSessionsCommon::CheckCreateSessionState] Could not create session with user [%s] not logged in"), *ToLogString(LocalAccountId));
 
 				return Errors::InvalidUser();
 			}
@@ -675,9 +675,9 @@ namespace UE::Online {
 
 			if (!FoundSession->SessionSettings.bIsDedicatedServerSession)
 			{
-				if (!Services.GetAuthInterface()->IsLoggedIn(Params.LocalUserId))
+				if (!Services.GetAuthInterface()->IsLoggedIn(Params.LocalAccountId))
 				{
-					UE_LOG(LogTemp, Warning, TEXT("[FSessionsNull::CheckUpdateSessionState] Could not update session with user [%s] not logged in"), *ToLogString(Params.LocalUserId));
+					UE_LOG(LogTemp, Warning, TEXT("[FSessionsNull::CheckUpdateSessionState] Could not update session with user [%s] not logged in"), *ToLogString(Params.LocalAccountId));
 
 					return Errors::InvalidUser();
 				}
@@ -708,15 +708,15 @@ namespace UE::Online {
 	FOnlineError FSessionsCommon::CheckFindSessionsState(const FFindSessions::Params& Params)
 	{
 		// User login check
-		if (!Services.GetAuthInterface()->IsLoggedIn(Params.LocalUserId))
+		if (!Services.GetAuthInterface()->IsLoggedIn(Params.LocalAccountId))
 		{
-			UE_LOG(LogTemp, Warning, TEXT("[FSessionsCommon::CheckFindSessionsState] Could not find sessions with user [%s] not logged in"), *ToLogString(Params.LocalUserId));
+			UE_LOG(LogTemp, Warning, TEXT("[FSessionsCommon::CheckFindSessionsState] Could not find sessions with user [%s] not logged in"), *ToLogString(Params.LocalAccountId));
 
 			return Errors::InvalidUser();
 		}
 
 		// Ongoing search check
-		if (const TSharedRef<TOnlineAsyncOp<FFindSessions>>* CurrentSessionSearchHandle = CurrentSessionSearchHandlesUserMap.Find(Params.LocalUserId))
+		if (const TSharedRef<TOnlineAsyncOp<FFindSessions>>* CurrentSessionSearchHandle = CurrentSessionSearchHandlesUserMap.Find(Params.LocalAccountId))
 		{
 			UE_LOG(LogTemp, Warning, TEXT("[FSessionsCommon::CheckFindSessionsState] Could not find sessions, search already in progress"));
 
@@ -741,7 +741,7 @@ namespace UE::Online {
 	FOnlineError FSessionsCommon::CheckStartMatchmakingState(const FStartMatchmaking::Params& Params)
 	{
 		// Check if a session with that name already exists
-		TOptional<FOnlineError> Result = CheckSessionExistsByName(Params.LocalUserId, Params.SessionName);
+		TOptional<FOnlineError> Result = CheckSessionExistsByName(Params.LocalAccountId, Params.SessionName);
 		if (!Result.IsSet()) // If CheckSessionExistsByName did not return an error, a session with that name already exists
 		{
 			UE_LOG(LogTemp, Warning, TEXT("[FSessionsCommon::CheckStartMatchmakingState] Could not join session with name [%s]. A session with that name already exists"), *Params.SessionName.ToString());
@@ -752,15 +752,15 @@ namespace UE::Online {
 		// User login check for all local users
 		IAuthPtr Auth = Services.GetAuthInterface();
 
-		TArray<FAccountId> LocalUserIds;
-		LocalUserIds.Reserve(Params.LocalUsers.Num());
-		Params.LocalUsers.GenerateKeyArray(LocalUserIds);
+		TArray<FAccountId> LocalAccountIds;
+		LocalAccountIds.Reserve(Params.LocalAccounts.Num());
+		Params.LocalAccounts.GenerateKeyArray(LocalAccountIds);
 
-		for (const FAccountId& LocalUserId : LocalUserIds)
+		for (const FAccountId& LocalAccountId : LocalAccountIds)
 		{
-			if (!Auth->IsLoggedIn(LocalUserId))
+			if (!Auth->IsLoggedIn(LocalAccountId))
 			{
-				UE_LOG(LogTemp, Warning, TEXT("[FSessionsCommon::CheckStartMatchmakingState] Could not join session with user [%s] not logged in"), *ToLogString(LocalUserId));
+				UE_LOG(LogTemp, Warning, TEXT("[FSessionsCommon::CheckStartMatchmakingState] Could not join session with user [%s] not logged in"), *ToLogString(LocalAccountId));
 
 				return Errors::InvalidUser();
 			}
@@ -773,16 +773,16 @@ namespace UE::Online {
 
 	FOnlineError FSessionsCommon::CheckJoinSessionParams(const FJoinSession::Params& Params)
 	{
-		if (!Params.LocalUserId.IsValid())
+		if (!Params.LocalAccountId.IsValid())
 		{
-			UE_LOG(LogTemp, Warning, TEXT("[FSessionsCommon::CheckJoinSessionParams] Could not join session with name [%s]. LocalUserId [%s] not valid"), *ToLogString(Params.LocalUserId));
+			UE_LOG(LogTemp, Warning, TEXT("[FSessionsCommon::CheckJoinSessionParams] Could not join session with name [%s]. LocalAccountId [%s] not valid"), *ToLogString(Params.LocalAccountId));
 
 			return Errors::InvalidParams();
 		}
 
-		if (!Params.LocalUsers.Contains(Params.LocalUserId))
+		if (!Params.LocalAccounts.Contains(Params.LocalAccountId))
 		{
-			UE_LOG(LogTemp, Warning, TEXT("[FSessionsCommon::CheckJoinSessionParams] Could not join session with name [%s]. Missing Session Member info for user [%s]"), *ToLogString(Params.LocalUserId));
+			UE_LOG(LogTemp, Warning, TEXT("[FSessionsCommon::CheckJoinSessionParams] Could not join session with name [%s]. Missing Session Member info for user [%s]"), *ToLogString(Params.LocalAccountId));
 
 			return Errors::InvalidParams();
 		}
@@ -807,7 +807,7 @@ namespace UE::Online {
 	FOnlineError FSessionsCommon::CheckJoinSessionState(const FJoinSession::Params& Params)
 	{
 		// Check if a session with that name already exists
-		TOptional<FOnlineError> Result = CheckSessionExistsByName(Params.LocalUserId, Params.SessionName);
+		TOptional<FOnlineError> Result = CheckSessionExistsByName(Params.LocalAccountId, Params.SessionName);
 		if (!Result.IsSet()) // If CheckSessionExistsByName did not return an error, a session with that name already exists
 		{
 			UE_LOG(LogTemp, Warning, TEXT("[FSessionsCommon::CheckJoinSessionState] Could not join session with name [%s]. A session with that name already exists"), *Params.SessionName.ToString());
@@ -818,22 +818,22 @@ namespace UE::Online {
 		// User login check for all local users
 		IAuthPtr Auth = Services.GetAuthInterface();
 
-		TArray<FAccountId> LocalUserIds;
-		LocalUserIds.Reserve(Params.LocalUsers.Num());
-		Params.LocalUsers.GenerateKeyArray(LocalUserIds);
+		TArray<FAccountId> LocalAccountIds;
+		LocalAccountIds.Reserve(Params.LocalAccounts.Num());
+		Params.LocalAccounts.GenerateKeyArray(LocalAccountIds);
 
-		for (const FAccountId& LocalUserId : LocalUserIds)
+		for (const FAccountId& LocalAccountId : LocalAccountIds)
 		{
-			if(!Auth->IsLoggedIn(LocalUserId))
+			if(!Auth->IsLoggedIn(LocalAccountId))
 			{
-				UE_LOG(LogTemp, Warning, TEXT("[FSessionsCommon::CheckJoinSessionState] Could not join session with user [%s] not logged in"), *ToLogString(LocalUserId));
+				UE_LOG(LogTemp, Warning, TEXT("[FSessionsCommon::CheckJoinSessionState] Could not join session with user [%s] not logged in"), *ToLogString(LocalAccountId));
 
 				return Errors::InvalidUser();
 			}
 		}
 
 		// We check that the session is cached and valid for a join operation by the users
-		TOnlineResult<FGetSessionById> GetSessionByIdResult = GetSessionById({ Params.LocalUserId, Params.SessionId });
+		TOnlineResult<FGetSessionById> GetSessionByIdResult = GetSessionById({ Params.LocalAccountId, Params.SessionId });
 		if (GetSessionByIdResult.IsError())
 		{
 			UE_LOG(LogTemp, Warning, TEXT("[FSessionsCommon::CheckJoinSessionState] Unable to find session with SessionId parameter [%s]. Please call FindSessions to get an updated list of available sessions "), *ToLogString(Params.SessionId));
@@ -845,11 +845,11 @@ namespace UE::Online {
 
 		const FSessionSettings& SessionSettings = FoundSession->SessionSettings;
 
-		for (const FAccountId& LocalUserId : LocalUserIds)
+		for (const FAccountId& LocalAccountId : LocalAccountIds)
 		{
-			if (SessionSettings.SessionMembers.Contains(LocalUserId))
+			if (SessionSettings.SessionMembers.Contains(LocalAccountId))
 			{
-				UE_LOG(LogTemp, Warning, TEXT("[FSessionsCommon::CheckJoinSessionState] Could not join session. User [%s] already in session"), *ToLogString(LocalUserId));
+				UE_LOG(LogTemp, Warning, TEXT("[FSessionsCommon::CheckJoinSessionState] Could not join session. User [%s] already in session"), *ToLogString(LocalAccountId));
 
 				return Errors::AccessDenied();
 			}
@@ -864,27 +864,27 @@ namespace UE::Online {
 
 		if (!SessionSettings.bAllowUnregisteredPlayers)
 		{
-			for (const FAccountId& LocalUserId : LocalUserIds)
+			for (const FAccountId& LocalAccountId : LocalAccountIds)
 			{
-				if (!SessionSettings.RegisteredPlayers.Contains(LocalUserId))
+				if (!SessionSettings.RegisteredPlayers.Contains(LocalAccountId))
 				{
-					UE_LOG(LogTemp, Warning, TEXT("[FSessionsCommon::CheckJoinSessionState] Could not join session. bAllowUnregisteredPlayers is set to false and user [%s] is not registered"), *ToLogString(LocalUserId));
+					UE_LOG(LogTemp, Warning, TEXT("[FSessionsCommon::CheckJoinSessionState] Could not join session. bAllowUnregisteredPlayers is set to false and user [%s] is not registered"), *ToLogString(LocalAccountId));
 
 					return Errors::AccessDenied();
 				}
 			}
 		}
 
-		if (SessionSettings.NumMaxPublicConnections > 0 && SessionSettings.NumOpenPublicConnections < (uint32)LocalUserIds.Num())
+		if (SessionSettings.NumMaxPublicConnections > 0 && SessionSettings.NumOpenPublicConnections < (uint32)LocalAccountIds.Num())
 		{
-			UE_LOG(LogTemp, Warning, TEXT("[FSessionsCommon::CheckJoinSessionState] Could not join session. Not enough NumOpenPublicConnections [%u] for the amount of joining players [%d]"), SessionSettings.NumOpenPublicConnections, LocalUserIds.Num());
+			UE_LOG(LogTemp, Warning, TEXT("[FSessionsCommon::CheckJoinSessionState] Could not join session. Not enough NumOpenPublicConnections [%u] for the amount of joining players [%d]"), SessionSettings.NumOpenPublicConnections, LocalAccountIds.Num());
 
 			return Errors::AccessDenied();
 		}
 
-		if (SessionSettings.NumMaxPrivateConnections > 0 && SessionSettings.NumOpenPrivateConnections < (uint32)LocalUserIds.Num())
+		if (SessionSettings.NumMaxPrivateConnections > 0 && SessionSettings.NumOpenPrivateConnections < (uint32)LocalAccountIds.Num())
 		{
-			UE_LOG(LogTemp, Warning, TEXT("[FSessionsCommon::CheckJoinSessionState] Could not join session. Not enough NumOpenPrivateConnections [%u] for the amount of joining players [%d]"), SessionSettings.NumOpenPrivateConnections, LocalUserIds.Num());
+			UE_LOG(LogTemp, Warning, TEXT("[FSessionsCommon::CheckJoinSessionState] Could not join session. Not enough NumOpenPrivateConnections [%u] for the amount of joining players [%d]"), SessionSettings.NumOpenPrivateConnections, LocalAccountIds.Num());
 
 			return Errors::AccessDenied();
 		}
@@ -894,7 +894,7 @@ namespace UE::Online {
 
 	FOnlineError FSessionsCommon::CheckAddSessionMembersState(const FAddSessionMembers::Params& Params)
 	{
-		if (TOptional<FOnlineError> Result = CheckSessionExistsByName(Params.LocalUserId, Params.SessionName))
+		if (TOptional<FOnlineError> Result = CheckSessionExistsByName(Params.LocalAccountId, Params.SessionName))
 		{
 			UE_LOG(LogTemp, Warning, TEXT("[FSessionsCommon::CheckAddSessionMembersState] Could not add session member to session with name [%s]. Session not found"), *Params.SessionName.ToString());
 
@@ -908,7 +908,7 @@ namespace UE::Online {
 
 	FOnlineError FSessionsCommon::CheckRemoveSessionMembersState(const FRemoveSessionMembers::Params& Params)
 	{
-		if (TOptional<FOnlineError> Result = CheckSessionExistsByName(Params.LocalUserId, Params.SessionName))
+		if (TOptional<FOnlineError> Result = CheckSessionExistsByName(Params.LocalAccountId, Params.SessionName))
 		{
 			UE_LOG(LogTemp, Warning, TEXT("[FSessionsCommon::CheckRemoveSessionMembersState] Could not remove session member from session with name [%s]. Session not found"), *Params.SessionName.ToString());
 
@@ -921,7 +921,7 @@ namespace UE::Online {
 	FOnlineError FSessionsCommon::CheckLeaveSessionState(const FLeaveSession::Params& Params)
 	{
 		// User login check for main caller, session check
-		if (TOptional<FOnlineError> Result = CheckSessionExistsByName(Params.LocalUserId, Params.SessionName))
+		if (TOptional<FOnlineError> Result = CheckSessionExistsByName(Params.LocalAccountId, Params.SessionName))
 		{
 			UE_LOG(LogTemp, Warning, TEXT("[FSessionsCommon::CheckLeaveSessionState] Could not leave session with name [%s]. Session not found"), *Params.SessionName.ToString());
 
@@ -931,11 +931,11 @@ namespace UE::Online {
 		IAuthPtr Auth = Services.GetAuthInterface();
 
 		// User login check for all local users
-		for (const FAccountId& LocalUserId : Params.LocalUsers)
+		for (const FAccountId& LocalAccountId : Params.LocalAccounts)
 		{
-			if (!Auth->IsLoggedIn(LocalUserId))
+			if (!Auth->IsLoggedIn(LocalAccountId))
 			{
-				UE_LOG(LogTemp, Warning, TEXT("[FSessionsCommon::CheckLeaveSessionState] Could not leave session with user [%s] not logged in"), *ToLogString(LocalUserId));
+				UE_LOG(LogTemp, Warning, TEXT("[FSessionsCommon::CheckLeaveSessionState] Could not leave session with user [%s] not logged in"), *ToLogString(LocalAccountId));
 
 				return Errors::InvalidUser();
 			}
@@ -954,9 +954,9 @@ namespace UE::Online {
 
 			if (!FoundSession->SessionSettings.bIsDedicatedServerSession)
 			{
-				if (!Services.GetAuthInterface()->IsLoggedIn(Params.LocalUserId))
+				if (!Services.GetAuthInterface()->IsLoggedIn(Params.LocalAccountId))
 				{
-					UE_LOG(LogTemp, Warning, TEXT("[FSessionsNull::CheckSendSessionInviteState] Could not send session invite with user [%s] not logged in"), *ToLogString(Params.LocalUserId));
+					UE_LOG(LogTemp, Warning, TEXT("[FSessionsNull::CheckSendSessionInviteState] Could not send session invite with user [%s] not logged in"), *ToLogString(Params.LocalAccountId));
 
 					return Errors::InvalidUser();
 				}
@@ -976,9 +976,9 @@ namespace UE::Online {
 
 	FOnlineError FSessionsCommon::CheckRejectSessionInviteState(const FRejectSessionInvite::Params& Params)
 	{
-		if (!Services.GetAuthInterface()->IsLoggedIn(Params.LocalUserId))
+		if (!Services.GetAuthInterface()->IsLoggedIn(Params.LocalAccountId))
 		{
-			UE_LOG(LogTemp, Warning, TEXT("[FSessionsNull::CheckRejectSessionInviteState] Could not send session invite with user [%s] not logged in"), *ToLogString(Params.LocalUserId));
+			UE_LOG(LogTemp, Warning, TEXT("[FSessionsNull::CheckRejectSessionInviteState] Could not send session invite with user [%s] not logged in"), *ToLogString(Params.LocalAccountId));
 
 			return Errors::InvalidUser();
 		}
@@ -990,7 +990,7 @@ namespace UE::Online {
 
 	FOnlineError FSessionsCommon::CheckRegisterPlayersState(const FRegisterPlayers::Params& Params)
 	{
-		if (TOptional<FOnlineError> Result = CheckSessionExistsByName(Params.LocalUserId, Params.SessionName))
+		if (TOptional<FOnlineError> Result = CheckSessionExistsByName(Params.LocalAccountId, Params.SessionName))
 		{
 			UE_LOG(LogTemp, Warning, TEXT("[FSessionsCommon::CheckRegisterPlayersState] Could not register players in session with name [%s]. Session not found"), *Params.SessionName.ToString());
 
@@ -1004,7 +1004,7 @@ namespace UE::Online {
 
 	FOnlineError FSessionsCommon::CheckUnregisterPlayersState(const FUnregisterPlayers::Params& Params)
 	{
-		if(TOptional<FOnlineError> Result = CheckSessionExistsByName(Params.LocalUserId, Params.SessionName))
+		if(TOptional<FOnlineError> Result = CheckSessionExistsByName(Params.LocalAccountId, Params.SessionName))
 		{
 			UE_LOG(LogTemp, Warning, TEXT("[FSessionsCommon::CheckUnregisterPlayersState] Could not unregister players from session with name [%s]. Session not found"), *Params.SessionName.ToString());
 
@@ -1015,7 +1015,7 @@ namespace UE::Online {
 	}
 
 	// TODO: Have all Check methods return TOptional too, change call sites, and write Macro for repeating code structure
-	TOptional<FOnlineError> FSessionsCommon::CheckSessionExistsByName(const FAccountId& LocalUserId, const FName& SessionName)
+	TOptional<FOnlineError> FSessionsCommon::CheckSessionExistsByName(const FAccountId& LocalAccountId, const FName& SessionName)
 	{
 		TOptional<FOnlineError> Result;
 
diff --git a/Engine/Plugins/Online/Experimental/OnlineServices/Source/OnlineServicesCommon/Private/Online/StatsCommon.cpp b/Engine/Plugins/Online/Experimental/OnlineServices/Source/OnlineServicesCommon/Private/Online/StatsCommon.cpp
--- a/Engine/Plugins/Online/Experimental/OnlineServices/Source/OnlineServicesCommon/Private/Online/StatsCommon.cpp
+++ b/Engine/Plugins/Online/Experimental/OnlineServices/Source/OnlineServicesCommon/Private/Online/StatsCommon.cpp
@@ -176,7 +176,7 @@ TOnlineResult<FGetCachedStats> FStatsCommon::GetCachedStats(FGetCachedStats::Par
 
 void FStatsCommon::CacheUserStats(const FUserStats& UserStats)
 {
-	if (FUserStats* ExistingUserStats = CachedUsersStats.FindByPredicate(FFindUserStatsByAccountId(UserStats.UserId)))
+	if (FUserStats* ExistingUserStats = CachedUsersStats.FindByPredicate(FFindUserStatsByAccountId(UserStats.AccountId)))
 	{
 		for (const TPair<FString, FStatValue>& StatPair : UserStats.Stats)
 		{
diff --git a/Engine/Plugins/Online/Experimental/OnlineServices/Source/OnlineServicesCommon/Public/Online/AchievementsCommon.h b/Engine/Plugins/Online/Experimental/OnlineServices/Source/OnlineServicesCommon/Public/Online/AchievementsCommon.h
--- a/Engine/Plugins/Online/Experimental/OnlineServices/Source/OnlineServicesCommon/Public/Online/AchievementsCommon.h
+++ b/Engine/Plugins/Online/Experimental/OnlineServices/Source/OnlineServicesCommon/Public/Online/AchievementsCommon.h
@@ -51,9 +51,9 @@ protected:
 	TOnlineEventCallable<void(const FAchievementStateUpdated&)> OnAchievementStateUpdatedEvent;
 
 	void UnlockAchievementsByStats(const FStatsUpdated& StatsUpdated);
-	void ExecuteUnlockRulesRelatedToStat(const FAccountId& UserId, const FString& StatName, const TMap<FString, FStatValue>& Stats, TArray<FString>& OutAchievementsToUnlock);
+	void ExecuteUnlockRulesRelatedToStat(const FAccountId& AccountId, const FString& StatName, const TMap<FString, FStatValue>& Stats, TArray<FString>& OutAchievementsToUnlock);
 	bool MeetUnlockCondition(FAchievementUnlockRule AchievementUnlockRule, const TMap<FString, FStatValue>& Stats);
-	bool IsUnlocked(const FAccountId& UserId, const FString& AchievementName) const;
+	bool IsUnlocked(const FAccountId& AccountId, const FString& AchievementName) const;
 
 	TArray<FAchievementUnlockRule> AchievementUnlockRules;
 	FOnlineEventDelegateHandle StatEventHandle;
diff --git a/Engine/Plugins/Online/Experimental/OnlineServices/Source/OnlineServicesCommon/Public/Online/AuthCommon.h b/Engine/Plugins/Online/Experimental/OnlineServices/Source/OnlineServicesCommon/Public/Online/AuthCommon.h
--- a/Engine/Plugins/Online/Experimental/OnlineServices/Source/OnlineServicesCommon/Public/Online/AuthCommon.h
+++ b/Engine/Plugins/Online/Experimental/OnlineServices/Source/OnlineServicesCommon/Public/Online/AuthCommon.h
@@ -15,7 +15,7 @@ public:
 	virtual ~FAccountInfoRegistry() = default;
 
 	TSharedPtr<FAccountInfo> Find(FPlatformUserId PlatformUserId) const;
-	TSharedPtr<FAccountInfo> Find(FAccountId AccountIdHandle) const;
+	TSharedPtr<FAccountInfo> Find(FAccountId AccountId) const;
 
 	TArray<TSharedRef<FAccountInfo>> GetAllAccountInfo(TFunction<bool(const TSharedRef<FAccountInfo>&)> Predicate) const;
 
diff --git a/Engine/Plugins/Online/Experimental/OnlineServices/Source/OnlineServicesCommon/Public/Online/LeaderboardsCommon.h b/Engine/Plugins/Online/Experimental/OnlineServices/Source/OnlineServicesCommon/Public/Online/LeaderboardsCommon.h
--- a/Engine/Plugins/Online/Experimental/OnlineServices/Source/OnlineServicesCommon/Public/Online/LeaderboardsCommon.h
+++ b/Engine/Plugins/Online/Experimental/OnlineServices/Source/OnlineServicesCommon/Public/Online/LeaderboardsCommon.h
@@ -47,7 +47,7 @@ struct FWriteLeaderboardScores
 
 	struct Params
 	{
-		FAccountId LocalUserId;
+		FAccountId LocalAccountId;
 		FString BoardName;
 		uint64 Score;
 	};
@@ -88,7 +88,7 @@ protected:
 namespace Meta {
 
 BEGIN_ONLINE_STRUCT_META(FWriteLeaderboardScores::Params)
-	ONLINE_STRUCT_FIELD(FWriteLeaderboardScores::Params, LocalUserId),
+	ONLINE_STRUCT_FIELD(FWriteLeaderboardScores::Params, LocalAccountId),
 	ONLINE_STRUCT_FIELD(FWriteLeaderboardScores::Params, BoardName),
 	ONLINE_STRUCT_FIELD(FWriteLeaderboardScores::Params, Score)
 END_ONLINE_STRUCT_META()
diff --git a/Engine/Plugins/Online/Experimental/OnlineServices/Source/OnlineServicesCommon/Public/Online/OnlineAsyncOpCache.h b/Engine/Plugins/Online/Experimental/OnlineServices/Source/OnlineServicesCommon/Public/Online/OnlineAsyncOpCache.h
--- a/Engine/Plugins/Online/Experimental/OnlineServices/Source/OnlineServicesCommon/Public/Online/OnlineAsyncOpCache.h
+++ b/Engine/Plugins/Online/Experimental/OnlineServices/Source/OnlineServicesCommon/Public/Online/OnlineAsyncOpCache.h
@@ -9,10 +9,10 @@
 
 namespace UE::Online {
 
-struct CLocalUserIdDefined
+struct CLocalAccountIdDefined
 {
 	template <typename T>
-	auto Requires() -> decltype(T::LocalUserId);
+	auto Requires() -> decltype(T::LocalAccountId);
 };
 
 enum class EOperationCacheExpirationPolicy : uint8
@@ -250,10 +250,10 @@ public:
 	{
 		TSharedRef<TOnlineAsyncOp<OpType>> Op = MakeShared<TOnlineAsyncOp<OpType>>(Services, MoveTemp(Params));
 
-		if constexpr (TModels<CLocalUserIdDefined, typename OpType::Params>::Value)
+		if constexpr (TModels<CLocalAccountIdDefined, typename OpType::Params>::Value)
 		{
-			// add LocalUserId to the Op Data
-			Op->Data.template Set<decltype(Params.LocalUserId)>(TEXT("LocalUserId"), Op->GetParams().LocalUserId);
+			// add LocalAccountId to the Op Data
+			Op->Data.template Set<decltype(Params.LocalAccountId)>(TEXT("LocalAccountId"), Op->GetParams().LocalAccountId);
 		}
 
 		return Op;
@@ -266,9 +266,9 @@ public:
 		TSharedPtr<TOnlineAsyncOp<OpType>> Op;
 
 		// find existing op
-		if constexpr (TModels<CLocalUserIdDefined, typename OpType::Params>::Value)
+		if constexpr (TModels<CLocalAccountIdDefined, typename OpType::Params>::Value)
 		{
-			if (TUniquePtr<IWrappedOperation>* OpPtr = UserOperations.FindOrAdd(Params.LocalUserId).Find(FWrappedOperationKey::Create<OpType, ParamsFuncsType>(Params)))
+			if (TUniquePtr<IWrappedOperation>* OpPtr = UserOperations.FindOrAdd(Params.LocalAccountId).Find(FWrappedOperationKey::Create<OpType, ParamsFuncsType>(Params)))
 			{
 				if (!(*OpPtr)->IsExpired())
 				{
@@ -299,9 +299,9 @@ public:
 					// remove from cache if operation was canceled
 					if (ThisOp.GetState() == EAsyncOpState::Cancelled)
 					{
-						if constexpr (TModels<CLocalUserIdDefined, typename OpType::Params>::Value)
+						if constexpr (TModels<CLocalAccountIdDefined, typename OpType::Params>::Value)
 						{
-							UserOperations.FindOrAdd(ThisOp.GetParams().LocalUserId).Remove(FWrappedOperationKey::Create<OpType, ParamsFuncsType>(ThisOp.GetParams()));
+							UserOperations.FindOrAdd(ThisOp.GetParams().LocalAccountId).Remove(FWrappedOperationKey::Create<OpType, ParamsFuncsType>(ThisOp.GetParams()));
 						}
 						else
 						{
@@ -321,9 +321,9 @@ public:
 		TSharedPtr<TOnlineAsyncOp<OpType>> Op;
 
 		// find existing op
-		if constexpr (TModels<CLocalUserIdDefined, typename OpType::Params>::Value)
+		if constexpr (TModels<CLocalAccountIdDefined, typename OpType::Params>::Value)
 		{
-			if (TUniquePtr<IWrappedOperation>* OpPtr = UserOperations.FindOrAdd(Params.LocalUserId).Find(FWrappedOperationKey::Create<OpType, ParamsFuncsType>(Params)))
+			if (TUniquePtr<IWrappedOperation>* OpPtr = UserOperations.FindOrAdd(Params.LocalAccountId).Find(FWrappedOperationKey::Create<OpType, ParamsFuncsType>(Params)))
 			{
 				Op = (*OpPtr)->GetRef<TSharedRef<TOnlineAsyncOp<OpType>>>();
 			}
@@ -350,9 +350,9 @@ public:
 			// remove from cache once operation has started. It is no longer mergeable at that point
 			FOnlineEventDelegateHandle Handle = Op->OnStart().Add(GetSharedThis(), [this](const TOnlineAsyncOp<OpType>& ThisOp)
 				{
-					if constexpr (TModels<CLocalUserIdDefined, typename OpType::Params>::Value)
+					if constexpr (TModels<CLocalAccountIdDefined, typename OpType::Params>::Value)
 					{
-						UserOperations.FindOrAdd(ThisOp.GetParams().LocalUserId).Remove(FWrappedOperationKey::Create<OpType, ParamsFuncsType>(ThisOp.GetParams()));
+						UserOperations.FindOrAdd(ThisOp.GetParams().LocalAccountId).Remove(FWrappedOperationKey::Create<OpType, ParamsFuncsType>(ThisOp.GetParams()));
 					}
 					else
 					{
@@ -383,12 +383,12 @@ private:
 
 		TSharedRef<TOnlineAsyncOp<OpType>> Op = WrappedOp->GetDataRef();
 
-		if constexpr (TModels<CLocalUserIdDefined, typename OpType::Params>::Value)
+		if constexpr (TModels<CLocalAccountIdDefined, typename OpType::Params>::Value)
 		{
-			// add LocalUserId to the Op Data
-			Op->Data.template Set<decltype(Params.LocalUserId)>(TEXT("LocalUserId"), Op->GetParams().LocalUserId);
+			// add LocalAccountId to the Op Data
+			Op->Data.template Set<decltype(Params.LocalAccountId)>(TEXT("LocalAccountId"), Op->GetParams().LocalAccountId);
 
-			UserOperations.FindOrAdd(Op->GetParams().LocalUserId).Add(FWrappedOperationKey::Create<OpType, ParamsFuncsType>(Op->GetParams()), MoveTemp(WrappedOp));
+			UserOperations.FindOrAdd(Op->GetParams().LocalAccountId).Add(FWrappedOperationKey::Create<OpType, ParamsFuncsType>(Op->GetParams()), MoveTemp(WrappedOp));
 		}
 		else
 		{
diff --git a/Engine/Plugins/Online/Experimental/OnlineServices/Source/OnlineServicesCommon/Public/Online/OnlineExecHandler.h b/Engine/Plugins/Online/Experimental/OnlineServices/Source/OnlineServicesCommon/Public/Online/OnlineExecHandler.h
--- a/Engine/Plugins/Online/Experimental/OnlineServices/Source/OnlineServicesCommon/Public/Online/OnlineExecHandler.h
+++ b/Engine/Plugins/Online/Experimental/OnlineServices/Source/OnlineServicesCommon/Public/Online/OnlineExecHandler.h
@@ -229,11 +229,11 @@ public:
 
 	static bool Assign(FString& VariantValue, VariantObject& Variant, IOnlineServices* Services)
 	{
-		FAccountId Handle;
+		FAccountId AccountId;
 		const TCHAR* ValueTCHAR = *VariantValue;
-		if (ParseOnlineExecParams<TOnlineIdHandle<OnlineIdHandleTags::FAccount>>(ValueTCHAR, Handle, Services))
+		if (ParseOnlineExecParams<TOnlineIdHandle<OnlineIdHandleTags::FAccount>>(ValueTCHAR, AccountId, Services))
 		{
-			Variant.template Set<FAccountId>(Handle);
+			Variant.template Set<FAccountId>(AccountId);
 			return true;
 		}
 		return false;
diff --git a/Engine/Plugins/Online/Experimental/OnlineServices/Source/OnlineServicesCommon/Public/Online/StatsCommon.h b/Engine/Plugins/Online/Experimental/OnlineServices/Source/OnlineServicesCommon/Public/Online/StatsCommon.h
--- a/Engine/Plugins/Online/Experimental/OnlineServices/Source/OnlineServicesCommon/Public/Online/StatsCommon.h
+++ b/Engine/Plugins/Online/Experimental/OnlineServices/Source/OnlineServicesCommon/Public/Online/StatsCommon.h
@@ -50,17 +50,17 @@ struct FStatDefinition
 
 struct FFindUserStatsByAccountId
 {
-	FFindUserStatsByAccountId(const FAccountId& InUserId)
-		: UserId(InUserId)
+	FFindUserStatsByAccountId(const FAccountId& InAccountId)
+		: AccountId(InAccountId)
 	{
 	}
 
 	bool operator()(const FUserStats& UserStats) const
 	{
-		return UserStats.UserId == UserId;
+		return UserStats.AccountId == AccountId;
 	}
 
-	FAccountId UserId;
+	FAccountId AccountId;
 };
 
 class ONLINESERVICESCOMMON_API FStatsCommon : public TOnlineComponent<IStats>
diff --git a/Engine/Plugins/Online/Experimental/OnlineServices/Source/OnlineServicesCommonEngineUtils/Private/Online/SessionsLAN.cpp b/Engine/Plugins/Online/Experimental/OnlineServices/Source/OnlineServicesCommonEngineUtils/Private/Online/SessionsLAN.cpp
--- a/Engine/Plugins/Online/Experimental/OnlineServices/Source/OnlineServicesCommonEngineUtils/Private/Online/SessionsLAN.cpp
+++ b/Engine/Plugins/Online/Experimental/OnlineServices/Source/OnlineServicesCommonEngineUtils/Private/Online/SessionsLAN.cpp
@@ -152,23 +152,23 @@ TOnlineAsyncOpHandle<FCreateSession> FSessionsLAN::CreateSession(FCreateSession:
 		}
 
 		TSharedRef<FSessionLAN> NewSessionLANRef = MakeShared<FSessionLAN>();
-		NewSessionLANRef->OwnerUserId = OpParams.LocalUserId;
+		NewSessionLANRef->OwnerAccountId = OpParams.LocalAccountId;
 		NewSessionLANRef->SessionId = FOnlineSessionIdRegistryLAN::GetChecked(Services.GetServicesProvider()).GetNextSessionId();
 		NewSessionLANRef->SessionSettings = OpParams.SessionSettings;
 
 		// For LAN sessions, we'll add all the Session members manually instead of calling JoinSession since there is no API calls involved
-		NewSessionLANRef->SessionSettings.SessionMembers.Append(OpParams.LocalUsers);
+		NewSessionLANRef->SessionSettings.SessionMembers.Append(OpParams.LocalAccounts);
 
 		// We save the local object for the session, and set up the appropriate references
 		AllSessionsById.Emplace(NewSessionLANRef->SessionId, NewSessionLANRef);
 
 		LocalSessionsByName.Emplace(OpParams.SessionName, NewSessionLANRef->SessionId);
 
-		NamedSessionUserMap.FindOrAdd(OpParams.LocalUserId).AddUnique(OpParams.SessionName);
+		NamedSessionUserMap.FindOrAdd(OpParams.LocalAccountId).AddUnique(OpParams.SessionName);
 
 		if (NewSessionLANRef->SessionSettings.bPresenceEnabled)
 		{
-			FOnlineSessionIdHandle& PresenceSessionId = PresenceSessionsUserMap.FindOrAdd(OpParams.LocalUserId);
+			FOnlineSessionIdHandle& PresenceSessionId = PresenceSessionsUserMap.FindOrAdd(OpParams.LocalAccountId);
 			PresenceSessionId = NewSessionLANRef->SessionId;
 		}
 
@@ -257,9 +257,9 @@ TOnlineAsyncOpHandle<FFindSessions> FSessionsLAN::FindSessions(FFindSessions::Pa
 			return;
 		}
 
-		CurrentSessionSearchHandlesUserMap.Emplace(OpParams.LocalUserId, Op.AsShared());
+		CurrentSessionSearchHandlesUserMap.Emplace(OpParams.LocalAccountId, Op.AsShared());
 
-		FindLANSessions(OpParams.LocalUserId);
+		FindLANSessions(OpParams.LocalAccountId);
 	})
 	.Enqueue(GetSerialQueue());
 
@@ -290,7 +290,7 @@ TOnlineAsyncOpHandle<FJoinSession> FSessionsLAN::JoinSession(FJoinSession::Param
 			return;
 		}
 
-		TOnlineResult<FGetSessionById> GetSessionByIdResult = GetSessionById({ OpParams.LocalUserId, OpParams.SessionId });
+		TOnlineResult<FGetSessionById> GetSessionByIdResult = GetSessionById({ OpParams.LocalAccountId, OpParams.SessionId });
 		if (GetSessionByIdResult.IsError())
 		{
 			// If no result is found, the id might be expired, which we should notify
@@ -311,21 +311,21 @@ TOnlineAsyncOpHandle<FJoinSession> FSessionsLAN::JoinSession(FJoinSession::Param
 
 		LocalSessionsByName.Emplace(OpParams.SessionName, NewSessionLANRef->SessionId);
 
-		NamedSessionUserMap.FindOrAdd(OpParams.LocalUserId).AddUnique(OpParams.SessionName);
+		NamedSessionUserMap.FindOrAdd(OpParams.LocalAccountId).AddUnique(OpParams.SessionName);
 
 		if (NewSessionLANRef->SessionSettings.bPresenceEnabled)
 		{
-			FOnlineSessionIdHandle& PresenceSessionId = PresenceSessionsUserMap.FindOrAdd(OpParams.LocalUserId);
+			FOnlineSessionIdHandle& PresenceSessionId = PresenceSessionsUserMap.FindOrAdd(OpParams.LocalAccountId);
 			PresenceSessionId = NewSessionLANRef->SessionId;
 		}
 
 		Op.SetResult(FJoinSession::Result{ });
 
-		TArray<FAccountId> LocalUserIds;
-		LocalUserIds.Reserve(OpParams.LocalUsers.Num());
-		OpParams.LocalUsers.GenerateKeyArray(LocalUserIds);
+		TArray<FAccountId> LocalAccountIds;
+		LocalAccountIds.Reserve(OpParams.LocalAccounts.Num());
+		OpParams.LocalAccounts.GenerateKeyArray(LocalAccountIds);
 
-		FSessionJoined SessionJoinedEvent = { MoveTemp(LocalUserIds), NewSessionLANRef->SessionId };
+		FSessionJoined SessionJoinedEvent = { MoveTemp(LocalAccountIds), NewSessionLANRef->SessionId };
 		
 		SessionEvents.OnSessionJoined.Broadcast(SessionJoinedEvent);
 	})
@@ -354,11 +354,11 @@ TOnlineAsyncOpHandle<FLeaveSession> FSessionsLAN::LeaveSession(FLeaveSession::Pa
 		{
 			TSharedRef<const FSession> FoundSession = GetSessionByNameResult.GetOkValue().Session;
 
-			NamedSessionUserMap.FindChecked(OpParams.LocalUserId).Remove(OpParams.SessionName);
+			NamedSessionUserMap.FindChecked(OpParams.LocalAccountId).Remove(OpParams.SessionName);
 
 			if (FoundSession->SessionSettings.bPresenceEnabled)
 			{
-				PresenceSessionsUserMap.Remove(OpParams.LocalUserId);
+				PresenceSessionsUserMap.Remove(OpParams.LocalAccountId);
 			}
 
 			ClearSessionByName(OpParams.SessionName);
@@ -368,7 +368,7 @@ TOnlineAsyncOpHandle<FLeaveSession> FSessionsLAN::LeaveSession(FLeaveSession::Pa
 		Op.SetResult(FLeaveSession::Result{ });
 
 		FSessionLeft SessionLeftEvent;
-		SessionLeftEvent.LocalUserIds.Append(OpParams.LocalUsers);
+		SessionLeftEvent.LocalAccountIds.Append(OpParams.LocalAccounts);
 		SessionEvents.OnSessionLeft.Broadcast(SessionLeftEvent);
 	})
 	.Enqueue(GetSerialQueue());
@@ -412,14 +412,14 @@ TOptional<FOnlineError> FSessionsLAN::TryHostLANSession()
 	return Result;
 }
 
-void FSessionsLAN::FindLANSessions(const FAccountId& LocalUserId)
+void FSessionsLAN::FindLANSessions(const FAccountId& LocalAccountId)
 {
 	// Recreate the unique identifier for this client
 	GenerateNonce((uint8*)&LANSessionManager->LanNonce, 8);
 
 	// Bind delegates
-	FOnValidResponsePacketDelegate ResponseDelegate = FOnValidResponsePacketDelegate::CreateThreadSafeSP(this, &FSessionsLAN::OnValidResponsePacketReceived, LocalUserId);
-	FOnSearchingTimeoutDelegate TimeoutDelegate = FOnSearchingTimeoutDelegate::CreateThreadSafeSP(this, &FSessionsLAN::OnLANSearchTimeout, LocalUserId);
+	FOnValidResponsePacketDelegate ResponseDelegate = FOnValidResponsePacketDelegate::CreateThreadSafeSP(this, &FSessionsLAN::OnValidResponsePacketReceived, LocalAccountId);
+	FOnSearchingTimeoutDelegate TimeoutDelegate = FOnSearchingTimeoutDelegate::CreateThreadSafeSP(this, &FSessionsLAN::OnLANSearchTimeout, LocalAccountId);
 
 	FNboSerializeToBuffer Packet(LAN_BEACON_MAX_PACKET_SIZE);
 	LANSessionManager->CreateClientQueryPacket(Packet, LANSessionManager->LanNonce);
@@ -433,9 +433,9 @@ void FSessionsLAN::FindLANSessions(const FAccountId& LocalUserId)
 		}
 
 		// Trigger the delegate as having failed
-		CurrentSessionSearchHandlesUserMap.FindChecked(LocalUserId)->SetError(Errors::RequestFailure()); // TODO: May need a new error type
+		CurrentSessionSearchHandlesUserMap.FindChecked(LocalAccountId)->SetError(Errors::RequestFailure()); // TODO: May need a new error type
 
-		CurrentSessionSearchHandlesUserMap.Remove(LocalUserId);
+		CurrentSessionSearchHandlesUserMap.Remove(LocalAccountId);
 
 		// If we were hosting public sessions before the search, we'll return the beacon to that state
 		if (PublicSessionsHosted > 0)
@@ -490,9 +490,9 @@ void FSessionsLAN::OnValidQueryPacketReceived(uint8* PacketData, int32 PacketLen
 	}
 }
 
-void FSessionsLAN::OnValidResponsePacketReceived(uint8* PacketData, int32 PacketLength, const FAccountId LocalUserId)
+void FSessionsLAN::OnValidResponsePacketReceived(uint8* PacketData, int32 PacketLength, const FAccountId LocalAccountId)
 {
-	if (TArray<FOnlineSessionIdHandle>* SearchResults = SearchResultsUserMap.Find(LocalUserId))
+	if (TArray<FOnlineSessionIdHandle>* SearchResults = SearchResultsUserMap.Find(LocalAccountId))
 	{
 		TSharedRef<FSessionLAN> Session = MakeShared<FSessionLAN>();
 
@@ -505,20 +505,20 @@ void FSessionsLAN::OnValidResponsePacketReceived(uint8* PacketData, int32 Packet
 	}
 }
 
-void FSessionsLAN::OnLANSearchTimeout(const FAccountId LocalUserId)
+void FSessionsLAN::OnLANSearchTimeout(const FAccountId LocalAccountId)
 {
 	if (LANSessionManager->GetBeaconState() == ELanBeaconState::Searching)
 	{
 		LANSessionManager->StopLANSession();
 	}
 
-	UE_LOG(LogTemp, Warning, TEXT("[FSessionsLAN::OnLANSearchTimeout] %d sessions found!"), SearchResultsUserMap.FindChecked(LocalUserId).Num());
+	UE_LOG(LogTemp, Warning, TEXT("[FSessionsLAN::OnLANSearchTimeout] %d sessions found!"), SearchResultsUserMap.FindChecked(LocalAccountId).Num());
 
-	TArray<FOnlineSessionIdHandle>& FoundSessionIds = SearchResultsUserMap.FindChecked(LocalUserId);
+	TArray<FOnlineSessionIdHandle>& FoundSessionIds = SearchResultsUserMap.FindChecked(LocalAccountId);
 
-	CurrentSessionSearchHandlesUserMap.FindChecked(LocalUserId)->SetResult({ FoundSessionIds });
+	CurrentSessionSearchHandlesUserMap.FindChecked(LocalAccountId)->SetResult({ FoundSessionIds });
 
-	CurrentSessionSearchHandlesUserMap.Remove(LocalUserId);
+	CurrentSessionSearchHandlesUserMap.Remove(LocalAccountId);
 
 	// If we were hosting public sessions before the search, we'll return the beacon to that state
 	if (PublicSessionsHosted > 0)
diff --git a/Engine/Plugins/Online/Experimental/OnlineServices/Source/OnlineServicesInterface/Private/Online/Achievements.cpp b/Engine/Plugins/Online/Experimental/OnlineServices/Source/OnlineServicesInterface/Private/Online/Achievements.cpp
--- a/Engine/Plugins/Online/Experimental/OnlineServices/Source/OnlineServicesInterface/Private/Online/Achievements.cpp
+++ b/Engine/Plugins/Online/Experimental/OnlineServices/Source/OnlineServicesInterface/Private/Online/Achievements.cpp
@@ -46,7 +46,7 @@ FString ToLogString(const FAchievementState& State)
 
 bool operator==(const FAchievementStateUpdated& A, const FAchievementStateUpdated& B)
 {
-	return A.LocalUserId == B.LocalUserId
+	return A.LocalAccountId == B.LocalAccountId
 		&& A.AchievementIds == B.AchievementIds;
 }
 
diff --git a/Engine/Plugins/Online/Experimental/OnlineServices/Source/OnlineServicesInterface/Private/Online/Sessions.cpp b/Engine/Plugins/Online/Experimental/OnlineServices/Source/OnlineServicesInterface/Private/Online/Sessions.cpp
--- a/Engine/Plugins/Online/Experimental/OnlineServices/Source/OnlineServicesInterface/Private/Online/Sessions.cpp
+++ b/Engine/Plugins/Online/Experimental/OnlineServices/Source/OnlineServicesInterface/Private/Online/Sessions.cpp
@@ -193,7 +193,7 @@ FSession::FSession()
 }
 
 FSession::FSession(const FSession& InSession)
-	: OwnerUserId(InSession.OwnerUserId)
+	: OwnerAccountId(InSession.OwnerAccountId)
 	, SessionId(InSession.SessionId)
 	, SessionSettings(InSession.SessionSettings)
 {
diff --git a/Engine/Plugins/Online/Experimental/OnlineServices/Source/OnlineServicesInterface/Public/Online/Achievements.h b/Engine/Plugins/Online/Experimental/OnlineServices/Source/OnlineServicesInterface/Public/Online/Achievements.h
--- a/Engine/Plugins/Online/Experimental/OnlineServices/Source/OnlineServicesInterface/Public/Online/Achievements.h
+++ b/Engine/Plugins/Online/Experimental/OnlineServices/Source/OnlineServicesInterface/Public/Online/Achievements.h
@@ -16,7 +16,7 @@ struct FQueryAchievementDefinitions
 	struct Params
 	{
 		/** Local user performing the operation */
-		FAccountId LocalUserId;
+		FAccountId LocalAccountId;
 	};
 
 	struct Result
@@ -31,7 +31,7 @@ struct FGetAchievementIds
 	struct Params
 	{
 		/** Local user performing the operation */
-		FAccountId LocalUserId;
+		FAccountId LocalAccountId;
 	};
 
 	struct Result
@@ -84,7 +84,7 @@ struct FGetAchievementDefinition
 	struct Params
 	{
 		/** Local user performing the operation */
-		FAccountId LocalUserId;
+		FAccountId LocalAccountId;
 		/** Achievement to get definition of */
 		FString AchievementId;
 	};
@@ -103,7 +103,7 @@ struct FQueryAchievementStates
 	struct Params
 	{
 		/** Local user performing the operation */
-		FAccountId LocalUserId;
+		FAccountId LocalAccountId;
 	};
 
 	struct Result
@@ -130,7 +130,7 @@ struct FGetAchievementState
 	struct Params
 	{
 		/** Local user performing the operation */
-		FAccountId LocalUserId;
+		FAccountId LocalAccountId;
 		/** Achievement to get state of */
 		FString AchievementId;
 	};
@@ -149,7 +149,7 @@ struct FUnlockAchievements
 	struct Params
 	{
 		/** Local user performing the operation */
-		FAccountId LocalUserId;
+		FAccountId LocalAccountId;
 		/** Achievements to unlock */
 		TArray<FString> AchievementIds;
 	};
@@ -166,7 +166,7 @@ struct FDisplayAchievementUI
 	struct Params
 	{
 		/** Local user performing the operation */
-		FAccountId LocalUserId;
+		FAccountId LocalAccountId;
 		/** Achievement to display */
 		FString AchievementId;
 	};
@@ -180,7 +180,7 @@ struct FDisplayAchievementUI
 struct FAchievementStateUpdated
 {
 	/** User whose achievement states have updated */
-	FAccountId LocalUserId;
+	FAccountId LocalAccountId;
 	/** Achievements which have updated */
 	TArray<FString> AchievementIds;
 };
@@ -239,14 +239,14 @@ public:
 namespace Meta {
 
 BEGIN_ONLINE_STRUCT_META(FQueryAchievementDefinitions::Params)
-	ONLINE_STRUCT_FIELD(FQueryAchievementDefinitions::Params, LocalUserId)
+	ONLINE_STRUCT_FIELD(FQueryAchievementDefinitions::Params, LocalAccountId)
 END_ONLINE_STRUCT_META()
 
 BEGIN_ONLINE_STRUCT_META(FQueryAchievementDefinitions::Result)
 END_ONLINE_STRUCT_META()
 
 BEGIN_ONLINE_STRUCT_META(FGetAchievementIds::Params)
-	ONLINE_STRUCT_FIELD(FGetAchievementIds::Params, LocalUserId)
+	ONLINE_STRUCT_FIELD(FGetAchievementIds::Params, LocalAccountId)
 END_ONLINE_STRUCT_META()
 
 BEGIN_ONLINE_STRUCT_META(FGetAchievementIds::Result)
@@ -272,7 +272,7 @@ BEGIN_ONLINE_STRUCT_META(FAchievementDefinition)
 END_ONLINE_STRUCT_META()
 
 BEGIN_ONLINE_STRUCT_META(FGetAchievementDefinition::Params)
-	ONLINE_STRUCT_FIELD(FGetAchievementDefinition::Params, LocalUserId),
+	ONLINE_STRUCT_FIELD(FGetAchievementDefinition::Params, LocalAccountId),
 	ONLINE_STRUCT_FIELD(FGetAchievementDefinition::Params, AchievementId)
 END_ONLINE_STRUCT_META()
 
@@ -281,7 +281,7 @@ BEGIN_ONLINE_STRUCT_META(FGetAchievementDefinition::Result)
 END_ONLINE_STRUCT_META()
 
 BEGIN_ONLINE_STRUCT_META(FQueryAchievementStates::Params)
-	ONLINE_STRUCT_FIELD(FQueryAchievementStates::Params, LocalUserId)
+	ONLINE_STRUCT_FIELD(FQueryAchievementStates::Params, LocalAccountId)
 END_ONLINE_STRUCT_META()
 
 BEGIN_ONLINE_STRUCT_META(FQueryAchievementStates::Result)
@@ -294,7 +294,7 @@ BEGIN_ONLINE_STRUCT_META(FAchievementState)
 END_ONLINE_STRUCT_META()
 
 BEGIN_ONLINE_STRUCT_META(FGetAchievementState::Params)
-	ONLINE_STRUCT_FIELD(FGetAchievementState::Params, LocalUserId),
+	ONLINE_STRUCT_FIELD(FGetAchievementState::Params, LocalAccountId),
 	ONLINE_STRUCT_FIELD(FGetAchievementState::Params, AchievementId)
 END_ONLINE_STRUCT_META()
 
@@ -303,7 +303,7 @@ BEGIN_ONLINE_STRUCT_META(FGetAchievementState::Result)
 END_ONLINE_STRUCT_META()
 
 BEGIN_ONLINE_STRUCT_META(FUnlockAchievements::Params)
-	ONLINE_STRUCT_FIELD(FUnlockAchievements::Params, LocalUserId),
+	ONLINE_STRUCT_FIELD(FUnlockAchievements::Params, LocalAccountId),
 	ONLINE_STRUCT_FIELD(FUnlockAchievements::Params, AchievementIds)
 END_ONLINE_STRUCT_META()
 
@@ -311,7 +311,7 @@ BEGIN_ONLINE_STRUCT_META(FUnlockAchievements::Result)
 END_ONLINE_STRUCT_META()
 
 BEGIN_ONLINE_STRUCT_META(FDisplayAchievementUI::Params)
-	ONLINE_STRUCT_FIELD(FDisplayAchievementUI::Params, LocalUserId),
+	ONLINE_STRUCT_FIELD(FDisplayAchievementUI::Params, LocalAccountId),
 	ONLINE_STRUCT_FIELD(FDisplayAchievementUI::Params, AchievementId)
 END_ONLINE_STRUCT_META()
 
@@ -319,7 +319,7 @@ BEGIN_ONLINE_STRUCT_META(FDisplayAchievementUI::Result)
 END_ONLINE_STRUCT_META()
 
 BEGIN_ONLINE_STRUCT_META(FAchievementStateUpdated)
-	ONLINE_STRUCT_FIELD(FAchievementStateUpdated, LocalUserId),
+	ONLINE_STRUCT_FIELD(FAchievementStateUpdated, LocalAccountId),
 	ONLINE_STRUCT_FIELD(FAchievementStateUpdated, AchievementIds)
 END_ONLINE_STRUCT_META()
 
diff --git a/Engine/Plugins/Online/Experimental/OnlineServices/Source/OnlineServicesInterface/Public/Online/Auth.h b/Engine/Plugins/Online/Experimental/OnlineServices/Source/OnlineServicesInterface/Public/Online/Auth.h
--- a/Engine/Plugins/Online/Experimental/OnlineServices/Source/OnlineServicesInterface/Public/Online/Auth.h
+++ b/Engine/Plugins/Online/Experimental/OnlineServices/Source/OnlineServicesInterface/Public/Online/Auth.h
@@ -182,7 +182,7 @@ struct FAuthLogout
 	struct Params
 	{
 		/** The online account id of the Local User making the request. */
-		FAccountId LocalUserId;
+		FAccountId LocalAccountId;
 		/** Whether to remove persistent credentials on logout. */
 		bool bDestroyAuth = false;
 	};
@@ -199,7 +199,7 @@ struct FAuthModifyAccountAttributes
 	struct Params
 	{
 		/** The online account id of the Local User making the request. */
-		FAccountId LocalUserId;
+		FAccountId LocalAccountId;
 		/** New or changed attributes. */
 		TMap<FSchemaAttributeId, FSchemaVariant> MutatedAttributes;
 		/** Attributes to be cleared. */
@@ -218,7 +218,7 @@ struct FAuthQueryExternalServerAuthTicket
 	struct Params
 	{
 		/** The online account id of the Local User making the request. */
-		FAccountId LocalUserId;
+		FAccountId LocalAccountId;
 	};
 
 	struct Result
@@ -234,7 +234,7 @@ struct FAuthQueryExternalAuthToken
 	struct Params
 	{
 		/** The online account id of the Local User making the request. */
-		FAccountId LocalUserId;
+		FAccountId LocalAccountId;
 		/** The method of external auth to provide. */
 		EExternalAuthTokenMethod Method = EExternalAuthTokenMethod::Primary;
 	};
@@ -252,7 +252,7 @@ struct FAuthQueryVerifiedAuthTicket
 	struct Params
 	{
 		/** The online account id of the Local User making the request. */
-		FAccountId LocalUserId;
+		FAccountId LocalAccountId;
 		/** The intended purpose of the auth ticket. */
 		ERemoteAuthTicketAudience Audience = ERemoteAuthTicketAudience::Peer;
 	};
@@ -273,7 +273,7 @@ struct FAuthCancelVerifiedAuthTicket
 	struct Params
 	{
 		/** The online account id of the Local User making the request. */
-		FAccountId LocalUserId;
+		FAccountId LocalAccountId;
 		/** Local ticket id used to reference the verified auth ticket. */
 		FOnlineVerifiedAuthTicketIdHandle VerifiedAuthTicketId;
 	};
@@ -322,7 +322,7 @@ struct FAuthGetLocalOnlineUserByOnlineAccountId
 	struct Params
 	{
 		/** Account Id of the Local User making the request */
-		FAccountId LocalUserId;
+		FAccountId LocalAccountId;
 	};
 
 	struct Result
@@ -562,7 +562,7 @@ BEGIN_ONLINE_STRUCT_META(FAuthLogin::Result)
 END_ONLINE_STRUCT_META()
 
 BEGIN_ONLINE_STRUCT_META(FAuthLogout::Params)
-	ONLINE_STRUCT_FIELD(FAuthLogout::Params, LocalUserId),
+	ONLINE_STRUCT_FIELD(FAuthLogout::Params, LocalAccountId),
 	ONLINE_STRUCT_FIELD(FAuthLogout::Params, bDestroyAuth)
 END_ONLINE_STRUCT_META()
 
@@ -570,7 +570,7 @@ BEGIN_ONLINE_STRUCT_META(FAuthLogout::Result)
 END_ONLINE_STRUCT_META()
 
 BEGIN_ONLINE_STRUCT_META(FAuthModifyAccountAttributes::Params)
-	ONLINE_STRUCT_FIELD(FAuthModifyAccountAttributes::Params, LocalUserId),
+	ONLINE_STRUCT_FIELD(FAuthModifyAccountAttributes::Params, LocalAccountId),
 	ONLINE_STRUCT_FIELD(FAuthModifyAccountAttributes::Params, MutatedAttributes),
 	ONLINE_STRUCT_FIELD(FAuthModifyAccountAttributes::Params, ClearedAttributes)
 END_ONLINE_STRUCT_META()
@@ -579,7 +579,7 @@ BEGIN_ONLINE_STRUCT_META(FAuthModifyAccountAttributes::Result)
 END_ONLINE_STRUCT_META()
 
 BEGIN_ONLINE_STRUCT_META(FAuthQueryExternalServerAuthTicket::Params)
-	ONLINE_STRUCT_FIELD(FAuthQueryExternalServerAuthTicket::Params, LocalUserId)
+	ONLINE_STRUCT_FIELD(FAuthQueryExternalServerAuthTicket::Params, LocalAccountId)
 END_ONLINE_STRUCT_META()
 
 BEGIN_ONLINE_STRUCT_META(FAuthQueryExternalServerAuthTicket::Result)
@@ -587,7 +587,7 @@ BEGIN_ONLINE_STRUCT_META(FAuthQueryExternalServerAuthTicket::Result)
 END_ONLINE_STRUCT_META()
 
 BEGIN_ONLINE_STRUCT_META(FAuthQueryExternalAuthToken::Params)
-	ONLINE_STRUCT_FIELD(FAuthQueryExternalAuthToken::Params, LocalUserId),
+	ONLINE_STRUCT_FIELD(FAuthQueryExternalAuthToken::Params, LocalAccountId),
 	ONLINE_STRUCT_FIELD(FAuthQueryExternalAuthToken::Params, Method)
 END_ONLINE_STRUCT_META()
 
@@ -596,7 +596,7 @@ BEGIN_ONLINE_STRUCT_META(FAuthQueryExternalAuthToken::Result)
 END_ONLINE_STRUCT_META()
 
 BEGIN_ONLINE_STRUCT_META(FAuthQueryVerifiedAuthTicket::Params)
-	ONLINE_STRUCT_FIELD(FAuthQueryVerifiedAuthTicket::Params, LocalUserId),
+	ONLINE_STRUCT_FIELD(FAuthQueryVerifiedAuthTicket::Params, LocalAccountId),
 	ONLINE_STRUCT_FIELD(FAuthQueryVerifiedAuthTicket::Params, Audience)
 END_ONLINE_STRUCT_META()
 
@@ -605,7 +605,7 @@ BEGIN_ONLINE_STRUCT_META(FAuthQueryVerifiedAuthTicket::Result)
 END_ONLINE_STRUCT_META()
 
 BEGIN_ONLINE_STRUCT_META(FAuthCancelVerifiedAuthTicket::Params)
-	ONLINE_STRUCT_FIELD(FAuthCancelVerifiedAuthTicket::Params, LocalUserId),
+	ONLINE_STRUCT_FIELD(FAuthCancelVerifiedAuthTicket::Params, LocalAccountId),
 	ONLINE_STRUCT_FIELD(FAuthCancelVerifiedAuthTicket::Params, VerifiedAuthTicketId)
 END_ONLINE_STRUCT_META()
 
@@ -629,7 +629,7 @@ BEGIN_ONLINE_STRUCT_META(FAuthEndVerifiedAuthSession::Result)
 END_ONLINE_STRUCT_META()
 
 BEGIN_ONLINE_STRUCT_META(FAuthGetLocalOnlineUserByOnlineAccountId::Params)
-	ONLINE_STRUCT_FIELD(FAuthGetLocalOnlineUserByOnlineAccountId::Params, LocalUserId)
+	ONLINE_STRUCT_FIELD(FAuthGetLocalOnlineUserByOnlineAccountId::Params, LocalAccountId)
 END_ONLINE_STRUCT_META()
 
 BEGIN_ONLINE_STRUCT_META(FAuthGetLocalOnlineUserByOnlineAccountId::Result)
diff --git a/Engine/Plugins/Online/Experimental/OnlineServices/Source/OnlineServicesInterface/Public/Online/Commerce.h b/Engine/Plugins/Online/Experimental/OnlineServices/Source/OnlineServicesInterface/Public/Online/Commerce.h
--- a/Engine/Plugins/Online/Experimental/OnlineServices/Source/OnlineServicesInterface/Public/Online/Commerce.h
+++ b/Engine/Plugins/Online/Experimental/OnlineServices/Source/OnlineServicesInterface/Public/Online/Commerce.h
@@ -66,7 +66,7 @@ struct FCommerceQueryOffers
 	struct Params
 	{
 		/* Local user */
-		FAccountId LocalUserId;
+		FAccountId LocalAccountId;
 	};
 
 	struct Result
@@ -83,7 +83,7 @@ struct FCommerceQueryOffersById
 	struct Params
 	{
 		/* Local user */
-		FAccountId LocalUserId;
+		FAccountId LocalAccountId;
 		/* List of offer ids */
 		TArray<FOfferId> OfferIds;
 	};
@@ -103,7 +103,7 @@ struct FCommerceGetOffers
 	struct Params
 	{
 		/* Local user */
-		FAccountId LocalUserId;
+		FAccountId LocalAccountId;
 	};
 
 	struct Result
@@ -121,7 +121,7 @@ struct FCommerceGetOffersById
 	struct Params
 	{
 		/* Local user */
-		FAccountId LocalUserId;
+		FAccountId LocalAccountId;
 		/* List of offer ids */
 		TArray<FOfferId> OfferIds;
 	};
@@ -142,7 +142,7 @@ struct FCommerceShowStoreUI
 	struct Params
 	{
 		/* Local user */
-		FAccountId LocalUserId;
+		FAccountId LocalAccountId;
 	};
 
 	struct Result
@@ -167,7 +167,7 @@ struct FCommerceCheckout
 	struct Params
 	{
 		/* Local user */
-		FAccountId LocalUserId;
+		FAccountId LocalAccountId;
 		/* Offers to purchase */
 		TArray<FPurchaseOffer> Offers;
 	};
@@ -206,7 +206,7 @@ struct FCommerceQueryTransactionEntitlements
 	struct Params
 	{
 		/* Local user */
-		FAccountId LocalUserId;
+		FAccountId LocalAccountId;
 		/* Transaction id for a completed purchase */
 		FString TransactionId;
 	};
@@ -226,7 +226,7 @@ struct FCommerceQueryEntitlements
 	struct Params
 	{
 		/* Local user */
-		FAccountId LocalUserId;
+		FAccountId LocalAccountId;
 		/* Whether or not to fetch redeemed entitlements */
 		bool bIncludeRedeemed;
 	};
@@ -245,7 +245,7 @@ struct FCommerceGetEntitlements
 	struct Params
 	{
 		/* Local user */
-		FAccountId LocalUserId;
+		FAccountId LocalAccountId;
 	};
 
 	struct Result
@@ -263,7 +263,7 @@ struct FCommerceRedeemEntitlement
 	struct Params
 	{
 		/* Local user */
-		FAccountId LocalUserId;
+		FAccountId LocalAccountId;
 		/* Entitlement id to mark as redeemed */
 		FEntitlementId EntitlementId;
 		/* Quantity to consume */
@@ -284,7 +284,7 @@ struct FCommerceRetrieveS2SToken
 	struct Params
 	{
 		/* Local user */
-		FAccountId LocalUserId;
+		FAccountId LocalAccountId;
 		/* Implementation specific, used to specify the type of token to retrieve. See implementation specific documentation for valid values */
 		FString TokenType;
 	};
@@ -302,7 +302,7 @@ struct FCommerceOnPurchaseComplete
 public:
 	
 	/* Local user */
-	FAccountId LocalUserId;
+	FAccountId LocalAccountId;
 	
 	/* If available, the transaction id for the completed purchase */
 	TOptional<FString> TransactionId;
@@ -384,13 +384,13 @@ BEGIN_ONLINE_STRUCT_META(FEntitlement)
 END_ONLINE_STRUCT_META()
 
 BEGIN_ONLINE_STRUCT_META(FCommerceQueryOffers::Params)
-	ONLINE_STRUCT_FIELD(FCommerceQueryOffers::Params, LocalUserId)
+	ONLINE_STRUCT_FIELD(FCommerceQueryOffers::Params, LocalAccountId)
 END_ONLINE_STRUCT_META()
 	BEGIN_ONLINE_STRUCT_META(FCommerceQueryOffers::Result)
 END_ONLINE_STRUCT_META()
 
 BEGIN_ONLINE_STRUCT_META(FCommerceQueryOffersById::Params)
-	ONLINE_STRUCT_FIELD(FCommerceQueryOffersById::Params, LocalUserId),
+	ONLINE_STRUCT_FIELD(FCommerceQueryOffersById::Params, LocalAccountId),
 	ONLINE_STRUCT_FIELD(FCommerceQueryOffersById::Params, OfferIds)
 END_ONLINE_STRUCT_META()
 BEGIN_ONLINE_STRUCT_META(FCommerceQueryOffersById::Result)
@@ -398,14 +398,14 @@ BEGIN_ONLINE_STRUCT_META(FCommerceQueryOffersById::Result)
 END_ONLINE_STRUCT_META()
 
 BEGIN_ONLINE_STRUCT_META(FCommerceGetOffers::Params)
-	ONLINE_STRUCT_FIELD(FCommerceGetOffers::Params, LocalUserId)
+	ONLINE_STRUCT_FIELD(FCommerceGetOffers::Params, LocalAccountId)
 END_ONLINE_STRUCT_META()
 BEGIN_ONLINE_STRUCT_META(FCommerceGetOffers::Result)
 	ONLINE_STRUCT_FIELD(FCommerceGetOffers::Result, Offers)
 END_ONLINE_STRUCT_META()
 
 BEGIN_ONLINE_STRUCT_META(FCommerceGetOffersById::Params)
-	ONLINE_STRUCT_FIELD(FCommerceGetOffersById::Params, LocalUserId),
+	ONLINE_STRUCT_FIELD(FCommerceGetOffersById::Params, LocalAccountId),
 	ONLINE_STRUCT_FIELD(FCommerceGetOffersById::Params, OfferIds)
 END_ONLINE_STRUCT_META()
 BEGIN_ONLINE_STRUCT_META(FCommerceGetOffersById::Result)
@@ -413,13 +413,13 @@ ONLINE_STRUCT_FIELD(FCommerceGetOffersById::Result, Offers)
 END_ONLINE_STRUCT_META()
 
 BEGIN_ONLINE_STRUCT_META(FCommerceShowStoreUI::Params)
-	ONLINE_STRUCT_FIELD(FCommerceShowStoreUI::Params, LocalUserId)
+	ONLINE_STRUCT_FIELD(FCommerceShowStoreUI::Params, LocalAccountId)
 END_ONLINE_STRUCT_META()
 BEGIN_ONLINE_STRUCT_META(FCommerceShowStoreUI::Result)
 END_ONLINE_STRUCT_META()
 
 BEGIN_ONLINE_STRUCT_META(FCommerceCheckout::Params)
-	ONLINE_STRUCT_FIELD(FCommerceCheckout::Params, LocalUserId),
+	ONLINE_STRUCT_FIELD(FCommerceCheckout::Params, LocalAccountId),
 	ONLINE_STRUCT_FIELD(FCommerceCheckout::Params, Offers)
 END_ONLINE_STRUCT_META()
 BEGIN_ONLINE_STRUCT_META(FCommerceCheckout::Result)
@@ -427,7 +427,7 @@ BEGIN_ONLINE_STRUCT_META(FCommerceCheckout::Result)
 END_ONLINE_STRUCT_META()
 
 BEGIN_ONLINE_STRUCT_META(FCommerceQueryTransactionEntitlements::Params)
-	ONLINE_STRUCT_FIELD(FCommerceQueryTransactionEntitlements::Params, LocalUserId),
+	ONLINE_STRUCT_FIELD(FCommerceQueryTransactionEntitlements::Params, LocalAccountId),
 	ONLINE_STRUCT_FIELD(FCommerceQueryTransactionEntitlements::Params, TransactionId)
 END_ONLINE_STRUCT_META()
 BEGIN_ONLINE_STRUCT_META(FCommerceQueryTransactionEntitlements::Result)
@@ -435,21 +435,21 @@ BEGIN_ONLINE_STRUCT_META(FCommerceQueryTransactionEntitlements::Result)
 END_ONLINE_STRUCT_META()
 
 BEGIN_ONLINE_STRUCT_META(FCommerceQueryEntitlements::Params)
-	ONLINE_STRUCT_FIELD(FCommerceQueryEntitlements::Params, LocalUserId),
+	ONLINE_STRUCT_FIELD(FCommerceQueryEntitlements::Params, LocalAccountId),
 	ONLINE_STRUCT_FIELD(FCommerceQueryEntitlements::Params, bIncludeRedeemed)
 END_ONLINE_STRUCT_META()
 BEGIN_ONLINE_STRUCT_META(FCommerceQueryEntitlements::Result)
 END_ONLINE_STRUCT_META()
 
 BEGIN_ONLINE_STRUCT_META(FCommerceGetEntitlements::Params)
-	ONLINE_STRUCT_FIELD(FCommerceGetEntitlements::Params, LocalUserId)
+	ONLINE_STRUCT_FIELD(FCommerceGetEntitlements::Params, LocalAccountId)
 END_ONLINE_STRUCT_META()
 BEGIN_ONLINE_STRUCT_META(FCommerceGetEntitlements::Result)
 	ONLINE_STRUCT_FIELD(FCommerceGetEntitlements::Result, Entitlements)
 END_ONLINE_STRUCT_META()
 
 BEGIN_ONLINE_STRUCT_META(FCommerceRedeemEntitlement::Params)
-	ONLINE_STRUCT_FIELD(FCommerceRedeemEntitlement::Params, LocalUserId),
+	ONLINE_STRUCT_FIELD(FCommerceRedeemEntitlement::Params, LocalAccountId),
 	ONLINE_STRUCT_FIELD(FCommerceRedeemEntitlement::Params, EntitlementId),
 	ONLINE_STRUCT_FIELD(FCommerceRedeemEntitlement::Params, Quantity)
 END_ONLINE_STRUCT_META()
@@ -457,7 +457,7 @@ BEGIN_ONLINE_STRUCT_META(FCommerceRedeemEntitlement::Result)
 END_ONLINE_STRUCT_META()
 
 BEGIN_ONLINE_STRUCT_META(FCommerceRetrieveS2SToken::Params)
-	ONLINE_STRUCT_FIELD(FCommerceRetrieveS2SToken::Params, LocalUserId),
+	ONLINE_STRUCT_FIELD(FCommerceRetrieveS2SToken::Params, LocalAccountId),
 	ONLINE_STRUCT_FIELD(FCommerceRetrieveS2SToken::Params, TokenType)
 END_ONLINE_STRUCT_META()
 BEGIN_ONLINE_STRUCT_META(FCommerceRetrieveS2SToken::Result)
@@ -466,7 +466,7 @@ BEGIN_ONLINE_STRUCT_META(FCommerceRetrieveS2SToken::Result)
 END_ONLINE_STRUCT_META()
 
 BEGIN_ONLINE_STRUCT_META(FCommerceOnPurchaseComplete)
-	ONLINE_STRUCT_FIELD(FCommerceOnPurchaseComplete, LocalUserId),
+	ONLINE_STRUCT_FIELD(FCommerceOnPurchaseComplete, LocalAccountId),
 	ONLINE_STRUCT_FIELD(FCommerceOnPurchaseComplete, TransactionId)
 END_ONLINE_STRUCT_META()
 
diff --git a/Engine/Plugins/Online/Experimental/OnlineServices/Source/OnlineServicesInterface/Public/Online/ExternalUI.h b/Engine/Plugins/Online/Experimental/OnlineServices/Source/OnlineServicesInterface/Public/Online/ExternalUI.h
--- a/Engine/Plugins/Online/Experimental/OnlineServices/Source/OnlineServicesInterface/Public/Online/ExternalUI.h
+++ b/Engine/Plugins/Online/Experimental/OnlineServices/Source/OnlineServicesInterface/Public/Online/ExternalUI.h
@@ -37,7 +37,7 @@ struct FExternalUIShowFriendsUI
 	struct Params
 	{
 		/** Local user id to show the friends UI for */
-		FAccountId LocalUserId;
+		FAccountId LocalAccountId;
 	};
 
 	struct Result
@@ -84,7 +84,7 @@ BEGIN_ONLINE_STRUCT_META(FExternalUIShowLoginUI::Result)
 END_ONLINE_STRUCT_META()
 
 BEGIN_ONLINE_STRUCT_META(FExternalUIShowFriendsUI::Params)
-	ONLINE_STRUCT_FIELD(FExternalUIShowFriendsUI::Params, LocalUserId)
+	ONLINE_STRUCT_FIELD(FExternalUIShowFriendsUI::Params, LocalAccountId)
 END_ONLINE_STRUCT_META()
 
 BEGIN_ONLINE_STRUCT_META(FExternalUIShowFriendsUI::Result)
diff --git a/Engine/Plugins/Online/Experimental/OnlineServices/Source/OnlineServicesInterface/Public/Online/Friends.h b/Engine/Plugins/Online/Experimental/OnlineServices/Source/OnlineServicesInterface/Public/Online/Friends.h
--- a/Engine/Plugins/Online/Experimental/OnlineServices/Source/OnlineServicesInterface/Public/Online/Friends.h
+++ b/Engine/Plugins/Online/Experimental/OnlineServices/Source/OnlineServicesInterface/Public/Online/Friends.h
@@ -27,7 +27,7 @@ const TCHAR* LexToString(EFriendInviteStatus);
 struct FFriend
 {
 	/** Id of the friend */
-	FAccountId UserId;
+	FAccountId AccountId;
 	/** Friendship invite status - if an invite is pending */
 	TOptional<EFriendInviteStatus> InviteStatus;
 };
@@ -40,7 +40,7 @@ struct FQueryFriends
 	struct Params
 	{
 		/** Account Id of the local user making the request */
-		FAccountId LocalUserId;
+		FAccountId LocalAccountId;
 	};
 
 	/**
@@ -60,7 +60,7 @@ struct FGetFriends
 	struct Params
 	{
 		/** Account Id of the local user making the request */
-		FAccountId LocalUserId;
+		FAccountId LocalAccountId;
 	};
 
 	/** Output struct for Friends::GetFriends */
@@ -79,7 +79,7 @@ struct FAddFriend
 	struct Params
 	{
 		/** Account Id of the local user making the request */
-		FAccountId LocalUserId;
+		FAccountId LocalAccountId;
 		/** Friend to add */
 		FAccountId FriendId;
 	};
@@ -96,7 +96,7 @@ struct FAddFriend
 struct FFriendsListUpdated
 {
 	/** Local user who's friends list changed */
-	FAccountId LocalUserId;
+	FAccountId LocalAccountId;
 	/** Account ids of friends that have been added */
 	TArrayView<const FAccountId> AddedFriends;
 	/** Account ids of friends that have been removed */
@@ -157,19 +157,19 @@ namespace Meta {
 // TODO: Move to Friends_Meta.inl file?
 
 BEGIN_ONLINE_STRUCT_META(FFriend)
-	ONLINE_STRUCT_FIELD(FFriend, UserId),
+	ONLINE_STRUCT_FIELD(FFriend, AccountId),
 	ONLINE_STRUCT_FIELD(FFriend, InviteStatus)
 END_ONLINE_STRUCT_META()
 
 BEGIN_ONLINE_STRUCT_META(FQueryFriends::Params)
-	ONLINE_STRUCT_FIELD(FQueryFriends::Params, LocalUserId)
+	ONLINE_STRUCT_FIELD(FQueryFriends::Params, LocalAccountId)
 END_ONLINE_STRUCT_META()
 
 BEGIN_ONLINE_STRUCT_META(FQueryFriends::Result)
 END_ONLINE_STRUCT_META()
 
 BEGIN_ONLINE_STRUCT_META(FGetFriends::Params)
-	ONLINE_STRUCT_FIELD(FGetFriends::Params, LocalUserId)
+	ONLINE_STRUCT_FIELD(FGetFriends::Params, LocalAccountId)
 END_ONLINE_STRUCT_META()
 
 BEGIN_ONLINE_STRUCT_META(FGetFriends::Result)
@@ -177,7 +177,7 @@ BEGIN_ONLINE_STRUCT_META(FGetFriends::Result)
 END_ONLINE_STRUCT_META()
 
 BEGIN_ONLINE_STRUCT_META(FAddFriend::Params)
-	ONLINE_STRUCT_FIELD(FAddFriend::Params, LocalUserId),
+	ONLINE_STRUCT_FIELD(FAddFriend::Params, LocalAccountId),
 	ONLINE_STRUCT_FIELD(FAddFriend::Params, FriendId)
 END_ONLINE_STRUCT_META()
 
@@ -186,7 +186,7 @@ BEGIN_ONLINE_STRUCT_META(FAddFriend::Result)
 END_ONLINE_STRUCT_META()
 
 BEGIN_ONLINE_STRUCT_META(FFriendsListUpdated)
-	ONLINE_STRUCT_FIELD(FFriendsListUpdated, LocalUserId),
+	ONLINE_STRUCT_FIELD(FFriendsListUpdated, LocalAccountId),
 	ONLINE_STRUCT_FIELD(FFriendsListUpdated, AddedFriends),
 	ONLINE_STRUCT_FIELD(FFriendsListUpdated, RemovedFriends),
 	ONLINE_STRUCT_FIELD(FFriendsListUpdated, UpdatedFriends)
diff --git a/Engine/Plugins/Online/Experimental/OnlineServices/Source/OnlineServicesInterface/Public/Online/Leaderboards.h b/Engine/Plugins/Online/Experimental/OnlineServices/Source/OnlineServicesInterface/Public/Online/Leaderboards.h
--- a/Engine/Plugins/Online/Experimental/OnlineServices/Source/OnlineServicesInterface/Public/Online/Leaderboards.h
+++ b/Engine/Plugins/Online/Experimental/OnlineServices/Source/OnlineServicesInterface/Public/Online/Leaderboards.h
@@ -13,7 +13,7 @@ namespace UE::Online {
 struct FLeaderboardEntry
 {
 	/* The account id of this leaderboard entry */
-	FAccountId UserId;
+	FAccountId AccountId;
 	/* The rank of this account */
 	int32 Rank;
 	/* The score of this account */
@@ -27,9 +27,9 @@ struct FReadEntriesForUsers
 	struct Params
 	{
 		/* Local user id */
-		FAccountId LocalUserId;
+		FAccountId LocalAccountId;
 		/* The account ids of users */
-		TArray<FAccountId> UserIds;
+		TArray<FAccountId> AccountIds;
 		/* The leaderboard name */
 		FString BoardName;
 	};
@@ -48,7 +48,7 @@ struct FReadEntriesAroundRank
 	struct Params
 	{
 		/* Local user id */
-		FAccountId LocalUserId;
+		FAccountId LocalAccountId;
 		/* The rank index to start read */
 		uint32 Rank;
 		/* How many entries to fetch */
@@ -71,9 +71,9 @@ struct FReadEntriesAroundUser
 	struct Params
 	{
 		/* Local user id */
-		FAccountId LocalUserId;
+		FAccountId LocalAccountId;
 		/* The account id of specified user */
-		FAccountId UserId;
+		FAccountId AccountId;
 		/* The offset to the rank of user */
 		int32 Offset;
 		/* How many entries to fetch */
@@ -119,14 +119,14 @@ public:
 namespace Meta {
 
 BEGIN_ONLINE_STRUCT_META(FLeaderboardEntry)
-	ONLINE_STRUCT_FIELD(FLeaderboardEntry, UserId),
+	ONLINE_STRUCT_FIELD(FLeaderboardEntry, AccountId),
 	ONLINE_STRUCT_FIELD(FLeaderboardEntry, Rank),
 	ONLINE_STRUCT_FIELD(FLeaderboardEntry, Score)
 END_ONLINE_STRUCT_META()
 
 BEGIN_ONLINE_STRUCT_META(FReadEntriesForUsers::Params)
-	ONLINE_STRUCT_FIELD(FReadEntriesForUsers::Params, LocalUserId),
-	ONLINE_STRUCT_FIELD(FReadEntriesForUsers::Params, UserIds),
+	ONLINE_STRUCT_FIELD(FReadEntriesForUsers::Params, LocalAccountId),
+	ONLINE_STRUCT_FIELD(FReadEntriesForUsers::Params, AccountIds),
 	ONLINE_STRUCT_FIELD(FReadEntriesForUsers::Params, BoardName)
 END_ONLINE_STRUCT_META()
 
@@ -135,7 +135,7 @@ BEGIN_ONLINE_STRUCT_META(FReadEntriesForUsers::Result)
 END_ONLINE_STRUCT_META()
 
 BEGIN_ONLINE_STRUCT_META(FReadEntriesAroundRank::Params)
-	ONLINE_STRUCT_FIELD(FReadEntriesAroundRank::Params, LocalUserId),
+	ONLINE_STRUCT_FIELD(FReadEntriesAroundRank::Params, LocalAccountId),
 	ONLINE_STRUCT_FIELD(FReadEntriesAroundRank::Params, Rank),
 	ONLINE_STRUCT_FIELD(FReadEntriesAroundRank::Params, Limit),
 	ONLINE_STRUCT_FIELD(FReadEntriesAroundRank::Params, BoardName)
@@ -146,8 +146,8 @@ BEGIN_ONLINE_STRUCT_META(FReadEntriesAroundRank::Result)
 END_ONLINE_STRUCT_META()
 
 BEGIN_ONLINE_STRUCT_META(FReadEntriesAroundUser::Params)
-	ONLINE_STRUCT_FIELD(FReadEntriesAroundUser::Params, LocalUserId),
-	ONLINE_STRUCT_FIELD(FReadEntriesAroundUser::Params, UserId),
+	ONLINE_STRUCT_FIELD(FReadEntriesAroundUser::Params, LocalAccountId),
+	ONLINE_STRUCT_FIELD(FReadEntriesAroundUser::Params, AccountId),
 	ONLINE_STRUCT_FIELD(FReadEntriesAroundUser::Params, Offset),
 	ONLINE_STRUCT_FIELD(FReadEntriesAroundUser::Params, Limit),
 	ONLINE_STRUCT_FIELD(FReadEntriesAroundUser::Params, BoardName)
diff --git a/Engine/Plugins/Online/Experimental/OnlineServices/Source/OnlineServicesInterface/Public/Online/Lobbies.h b/Engine/Plugins/Online/Experimental/OnlineServices/Source/OnlineServicesInterface/Public/Online/Lobbies.h
--- a/Engine/Plugins/Online/Experimental/OnlineServices/Source/OnlineServicesInterface/Public/Online/Lobbies.h
+++ b/Engine/Plugins/Online/Experimental/OnlineServices/Source/OnlineServicesInterface/Public/Online/Lobbies.h
@@ -77,7 +77,7 @@ struct FLobby
 struct FJoinLobbyLocalUserData
 {
 	/** Local users who will be joining the lobby. */
-	FAccountId LocalUserId;
+	FAccountId LocalAccountId;
 
 	/** Initial attributes. */
 	TMap<FLobbyAttributeId, FLobbyVariant> Attributes;
@@ -105,7 +105,7 @@ struct FCreateLobby
 	struct Params
 	{
 		/** The local user agent which will perform the action. */
-		FAccountId LocalUserId;
+		FAccountId LocalAccountId;
 
 		/** The local name for the lobby. */
 		FName LocalName;
@@ -142,7 +142,7 @@ struct FFindLobbies
 	struct Params
 	{
 		/** The local user agent which will perform the action. */
-		FAccountId LocalUserId;
+		FAccountId LocalAccountId;
 
 		/**
 		* Max results to return in one search.
@@ -190,7 +190,7 @@ struct FJoinLobby
 	struct Params
 	{
 		/** The local user agent which will perform the action. */
-		FAccountId LocalUserId;
+		FAccountId LocalAccountId;
 
 		/** The local name for the lobby. */
 		FName LocalName;
@@ -218,7 +218,7 @@ struct FLeaveLobby
 	struct Params
 	{
 		/** The local user agent which will perform the action. */
-		FAccountId LocalUserId;
+		FAccountId LocalAccountId;
 
 		/** Id of the lobby to leave. */
 		FOnlineLobbyIdHandle LobbyId;
@@ -238,13 +238,13 @@ struct FInviteLobbyMember
 	struct Params
 	{
 		/** The local user agent which will perform the action. */
-		FAccountId LocalUserId;
+		FAccountId LocalAccountId;
 
 		/** Id of the lobby for which the invitation will be sent. */
 		FOnlineLobbyIdHandle LobbyId;
 
 		/** Id of the player who will be sent the invitation. */
-		FAccountId TargetUserId;
+		FAccountId TargetAccountId;
 	};
 
 	/** Output struct for Lobbies::InviteLobbyMember */
@@ -261,7 +261,7 @@ struct FDeclineLobbyInvitation
 	struct Params
 	{
 		/** The local user agent which will perform the action. */
-		FAccountId LocalUserId;
+		FAccountId LocalAccountId;
 
 		/** Id of the lobby for which the invitations will be declined. */
 		FOnlineLobbyIdHandle LobbyId;
@@ -281,13 +281,13 @@ struct FKickLobbyMember
 	struct Params
 	{
 		/** The local user agent which will perform the action. */
-		FAccountId LocalUserId;
+		FAccountId LocalAccountId;
 
 		/** Id of the lobby. */
 		FOnlineLobbyIdHandle LobbyId;
 
 		/** The target user to be kicked. */
-		FAccountId TargetUserId;
+		FAccountId TargetAccountId;
 	};
 
 	/** Output struct for Lobbies::KickLobbyMember */
@@ -304,13 +304,13 @@ struct FPromoteLobbyMember
 	struct Params
 	{
 		/** The local user agent which will perform the action. */
-		FAccountId LocalUserId;
+		FAccountId LocalAccountId;
 
 		/** Id of the lobby. */
 		FOnlineLobbyIdHandle LobbyId;
 
 		/** The target user to be promoted to owner. */
-		FAccountId TargetUserId;
+		FAccountId TargetAccountId;
 	};
 
 	/** Output struct for Lobbies::PromoteLobbyMember */
@@ -342,7 +342,7 @@ struct FModifyLobbyJoinPolicy
 	struct Params
 	{
 		/** The local user agent which will perform the action. */
-		FAccountId LocalUserId;
+		FAccountId LocalAccountId;
 
 		/** Id of the lobby. */
 		FOnlineLobbyIdHandle LobbyId;
@@ -365,7 +365,7 @@ struct FModifyLobbyAttributes
 	struct Params
 	{
 		/** The local user agent which will perform the action. */
-		FAccountId LocalUserId;
+		FAccountId LocalAccountId;
 
 		/** Id of the lobby. */
 		FOnlineLobbyIdHandle LobbyId;
@@ -391,7 +391,7 @@ struct FModifyLobbyMemberAttributes
 	struct Params
 	{
 		/** The local user agent which will perform the action. */
-		FAccountId LocalUserId;
+		FAccountId LocalAccountId;
 
 		/** Id of the lobby. */
 		FOnlineLobbyIdHandle LobbyId;
@@ -417,7 +417,7 @@ struct FGetJoinedLobbies
 	struct Params
 	{
 		/** The local user agent which will perform the action. */
-		FAccountId LocalUserId;
+		FAccountId LocalAccountId;
 	};
 
 	/** Output struct for Lobbies::GetJoinedLobbies */
@@ -435,7 +435,7 @@ struct FGetReceivedInvitations
 	struct Params
 	{
 		/** The local user agent which will perform the action. */
-		FAccountId LocalUserId;
+		FAccountId LocalAccountId;
 	};
 
 	/** Output struct for Lobbies::GetReceivedInvitations */
@@ -526,7 +526,7 @@ struct FLobbyMemberAttributesChanged
 struct FLobbyInvitationAdded
 {
 	/** The local user associated with the invitation. */
-	FAccountId LocalUserId;
+	FAccountId LocalAccountId;
 
 	/** The user who sent the invitation. */
 	FAccountId SenderId;
@@ -539,7 +539,7 @@ struct FLobbyInvitationAdded
 struct FLobbyInvitationRemoved
 {
 	/** The local user associated with the invitation. */
-	FAccountId LocalUserId;
+	FAccountId LocalAccountId;
 
 	/** The user who sent the invitation. */
 	FAccountId SenderId;
@@ -564,7 +564,7 @@ ONLINESERVICESINTERFACE_API void LexFromString(EUILobbyJoinRequestedSource& OutU
 struct FUILobbyJoinRequested
 {
 	/** The local user associated with the join request. */
-	FAccountId LocalUserId;
+	FAccountId LocalAccountId;
 
 	/** The lobby the local user requested to join, or the online error if there was a failure */
 	TResult<TSharedRef<const FLobby>, FOnlineError> Result;
@@ -809,12 +809,12 @@ BEGIN_ONLINE_STRUCT_META(FLobby)
 END_ONLINE_STRUCT_META()
 
 BEGIN_ONLINE_STRUCT_META(FJoinLobbyLocalUserData)
-	ONLINE_STRUCT_FIELD(FJoinLobbyLocalUserData, LocalUserId),
+	ONLINE_STRUCT_FIELD(FJoinLobbyLocalUserData, LocalAccountId),
 	ONLINE_STRUCT_FIELD(FJoinLobbyLocalUserData, Attributes)
 END_ONLINE_STRUCT_META()
 
 BEGIN_ONLINE_STRUCT_META(FCreateLobby::Params)
-	ONLINE_STRUCT_FIELD(FCreateLobby::Params, LocalUserId),
+	ONLINE_STRUCT_FIELD(FCreateLobby::Params, LocalAccountId),
 	ONLINE_STRUCT_FIELD(FCreateLobby::Params, LocalName),
 	ONLINE_STRUCT_FIELD(FCreateLobby::Params, SchemaName),
 	ONLINE_STRUCT_FIELD(FCreateLobby::Params, MaxMembers),
@@ -846,7 +846,7 @@ BEGIN_ONLINE_STRUCT_META(FJoinLobby::Result)
 END_ONLINE_STRUCT_META()
 
 BEGIN_ONLINE_STRUCT_META(FLeaveLobby::Params)
-	ONLINE_STRUCT_FIELD(FLeaveLobby::Params, LocalUserId),
+	ONLINE_STRUCT_FIELD(FLeaveLobby::Params, LocalAccountId),
 	ONLINE_STRUCT_FIELD(FLeaveLobby::Params, LobbyId)
 END_ONLINE_STRUCT_META()
 
@@ -902,7 +902,7 @@ BEGIN_ONLINE_STRUCT_META(FModifyLobbyMemberAttributes::Result)
 END_ONLINE_STRUCT_META()
 
 BEGIN_ONLINE_STRUCT_META(FGetJoinedLobbies::Params)
-	ONLINE_STRUCT_FIELD(FGetJoinedLobbies::Params, LocalUserId)
+	ONLINE_STRUCT_FIELD(FGetJoinedLobbies::Params, LocalAccountId)
 END_ONLINE_STRUCT_META()
 
 BEGIN_ONLINE_STRUCT_META(FGetJoinedLobbies::Result)
@@ -910,7 +910,7 @@ BEGIN_ONLINE_STRUCT_META(FGetJoinedLobbies::Result)
 END_ONLINE_STRUCT_META()
 
 BEGIN_ONLINE_STRUCT_META(FGetReceivedInvitations::Params)
-	ONLINE_STRUCT_FIELD(FGetReceivedInvitations::Params, LocalUserId)
+	ONLINE_STRUCT_FIELD(FGetReceivedInvitations::Params, LocalAccountId)
 END_ONLINE_STRUCT_META()
 
 BEGIN_ONLINE_STRUCT_META(FGetReceivedInvitations::Result)
@@ -952,19 +952,19 @@ BEGIN_ONLINE_STRUCT_META(FLobbyMemberAttributesChanged)
 END_ONLINE_STRUCT_META()
 
 BEGIN_ONLINE_STRUCT_META(FLobbyInvitationAdded)
-	ONLINE_STRUCT_FIELD(FLobbyInvitationAdded, LocalUserId),
+	ONLINE_STRUCT_FIELD(FLobbyInvitationAdded, LocalAccountId),
 	ONLINE_STRUCT_FIELD(FLobbyInvitationAdded, SenderId),
 	ONLINE_STRUCT_FIELD(FLobbyInvitationAdded, Lobby)
 END_ONLINE_STRUCT_META()
 
 BEGIN_ONLINE_STRUCT_META(FLobbyInvitationRemoved)
-	ONLINE_STRUCT_FIELD(FLobbyInvitationRemoved, LocalUserId),
+	ONLINE_STRUCT_FIELD(FLobbyInvitationRemoved, LocalAccountId),
 	ONLINE_STRUCT_FIELD(FLobbyInvitationRemoved, SenderId),
 	ONLINE_STRUCT_FIELD(FLobbyInvitationRemoved, Lobby)
 END_ONLINE_STRUCT_META()
 
 BEGIN_ONLINE_STRUCT_META(FUILobbyJoinRequested)
-	ONLINE_STRUCT_FIELD(FUILobbyJoinRequested, LocalUserId),
+	ONLINE_STRUCT_FIELD(FUILobbyJoinRequested, LocalAccountId),
 	ONLINE_STRUCT_FIELD(FUILobbyJoinRequested, Result)
 END_ONLINE_STRUCT_META()
 
diff --git a/Engine/Plugins/Online/Experimental/OnlineServices/Source/OnlineServicesInterface/Public/Online/OnlineServices.h b/Engine/Plugins/Online/Experimental/OnlineServices/Source/OnlineServicesInterface/Public/Online/OnlineServices.h
--- a/Engine/Plugins/Online/Experimental/OnlineServices/Source/OnlineServicesInterface/Public/Online/OnlineServices.h
+++ b/Engine/Plugins/Online/Experimental/OnlineServices/Source/OnlineServicesInterface/Public/Online/OnlineServices.h
@@ -35,7 +35,7 @@ struct FGetResolvedConnectString
 
 	struct Params
 	{
-		FAccountId LocalUserId;
+		FAccountId LocalAccountId;
 		FOnlineLobbyIdHandle LobbyId;
 		FOnlineSessionIdHandle SessionId;
 		FName PortType;
@@ -207,7 +207,7 @@ namespace Meta {
 // TODO: Move to OnlineServices_Meta.inl file?
 
 BEGIN_ONLINE_STRUCT_META(FGetResolvedConnectString::Params)
-	ONLINE_STRUCT_FIELD(FGetResolvedConnectString::Params, LocalUserId),
+	ONLINE_STRUCT_FIELD(FGetResolvedConnectString::Params, LocalAccountId),
 	ONLINE_STRUCT_FIELD(FGetResolvedConnectString::Params, LobbyId),
 	ONLINE_STRUCT_FIELD(FGetResolvedConnectString::Params, PortType)
 END_ONLINE_STRUCT_META()
