OSSv2: Rename Handle, UserId to AccountId when use FAccountId.

diff --git a/Engine/Plugins/Online/Experimental/OnlineServices/Source/OnlineServicesInterface/Public/Online/Presence.h b/Engine/Plugins/Online/Experimental/OnlineServices/Source/OnlineServicesInterface/Public/Online/Presence.h
--- a/Engine/Plugins/Online/Experimental/OnlineServices/Source/OnlineServicesInterface/Public/Online/Presence.h
+++ b/Engine/Plugins/Online/Experimental/OnlineServices/Source/OnlineServicesInterface/Public/Online/Presence.h
@@ -78,7 +78,7 @@ typedef TMap<FString, FPresenceVariant> FPresenceProperties;
 struct FUserPresence
 {
 	/** User whose presence this is */
-	FAccountId UserId;
+	FAccountId AccountId;
 	/** Presence state */
 	EUserPresenceStatus Status = EUserPresenceStatus::Unknown;
 	/** Session state */
@@ -100,9 +100,9 @@ struct FQueryPresence
 	struct Params
 	{
 		/** Local user performing the query */
-		FAccountId LocalUserId;
+		FAccountId LocalAccountId;
 		/** User to query the presence for */
-		FAccountId TargetUserId;
+		FAccountId TargetAccountId;
 		/** If true, then future presence updates for this user will be propagated via the OnPresenceUpdated event */
 		bool bListenToChanges = true;
 	};
@@ -122,9 +122,9 @@ struct FBatchQueryPresence
 	struct Params
 	{
 		/** Local user performing the query */
-		FAccountId LocalUserId;
+		FAccountId LocalAccountId;
 		/** Users to query the presence for */
-		TArray<FAccountId> TargetUserIds;
+		TArray<FAccountId> TargetAccountIds;
 		/** If true, then future presence updates for these users will be propagated via the OnPresenceUpdated event */
 		bool bListenToChanges = true;
 	};
@@ -143,9 +143,9 @@ struct FGetCachedPresence
 	struct Params
 	{
 		/** Local user getting the presence */
-		FAccountId LocalUserId;
+		FAccountId LocalAccountId;
 		/** User to get the presence for */
-		FAccountId TargetUserId;
+		FAccountId TargetAccountId;
 	};
 
 	struct Result
@@ -162,7 +162,7 @@ struct FUpdatePresence
 	struct Params
 	{
 		/** Local user performing the query */
-		FAccountId LocalUserId;
+		FAccountId LocalAccountId;
 		/** The new presence to send to the server */
 		TSharedRef<FUserPresence> Presence;
 	};
@@ -179,7 +179,7 @@ struct FPartialUpdatePresence
 	struct Params
 	{
 		/** Local user performing the query */
-		FAccountId LocalUserId;
+		FAccountId LocalAccountId;
 
 
 		struct FMutations
@@ -293,7 +293,7 @@ inline TSharedRef<FUserPresence> ApplyPresenceMutations(const FUserPresence& Bas
 struct FPresenceUpdated
 {
 	/** Local user receiving the presence update */
-	FAccountId LocalUserId;
+	FAccountId LocalAccountId;
 	/** Presence that has updated */
 	TSharedRef<const FUserPresence> UpdatedPresence;
 };
@@ -354,7 +354,7 @@ namespace Meta {
 // TODO: Move to Presence_Meta.inl file?
 
 BEGIN_ONLINE_STRUCT_META(FUserPresence)
-	ONLINE_STRUCT_FIELD(FUserPresence, UserId),
+	ONLINE_STRUCT_FIELD(FUserPresence, AccountId),
 	ONLINE_STRUCT_FIELD(FUserPresence, Status),
 	ONLINE_STRUCT_FIELD(FUserPresence, Joinability),
 	ONLINE_STRUCT_FIELD(FUserPresence, GameStatus),
@@ -364,14 +364,14 @@ BEGIN_ONLINE_STRUCT_META(FUserPresence)
 END_ONLINE_STRUCT_META()
 
 BEGIN_ONLINE_STRUCT_META(FQueryPresence::Params)
-	ONLINE_STRUCT_FIELD(FQueryPresence::Params, LocalUserId),
-	ONLINE_STRUCT_FIELD(FQueryPresence::Params, TargetUserId),
+	ONLINE_STRUCT_FIELD(FQueryPresence::Params, LocalAccountId),
+	ONLINE_STRUCT_FIELD(FQueryPresence::Params, TargetAccountId),
 	ONLINE_STRUCT_FIELD(FQueryPresence::Params, bListenToChanges)
 END_ONLINE_STRUCT_META()
 
 BEGIN_ONLINE_STRUCT_META(FBatchQueryPresence::Params)
-	ONLINE_STRUCT_FIELD(FBatchQueryPresence::Params, LocalUserId),
-	ONLINE_STRUCT_FIELD(FBatchQueryPresence::Params, TargetUserIds),
+	ONLINE_STRUCT_FIELD(FBatchQueryPresence::Params, LocalAccountId),
+	ONLINE_STRUCT_FIELD(FBatchQueryPresence::Params, TargetAccountIds),
 	ONLINE_STRUCT_FIELD(FBatchQueryPresence::Params, bListenToChanges)
 END_ONLINE_STRUCT_META()
 
@@ -384,7 +384,7 @@ BEGIN_ONLINE_STRUCT_META(FBatchQueryPresence::Result)
 END_ONLINE_STRUCT_META()
 
 BEGIN_ONLINE_STRUCT_META(FGetCachedPresence::Params)
-	ONLINE_STRUCT_FIELD(FGetCachedPresence::Params, LocalUserId)
+	ONLINE_STRUCT_FIELD(FGetCachedPresence::Params, LocalAccountId)
 END_ONLINE_STRUCT_META()
 
 BEGIN_ONLINE_STRUCT_META(FGetCachedPresence::Result)
@@ -392,7 +392,7 @@ BEGIN_ONLINE_STRUCT_META(FGetCachedPresence::Result)
 END_ONLINE_STRUCT_META()
 
 BEGIN_ONLINE_STRUCT_META(FUpdatePresence::Params)
-	ONLINE_STRUCT_FIELD(FUpdatePresence::Params, LocalUserId),
+	ONLINE_STRUCT_FIELD(FUpdatePresence::Params, LocalAccountId),
 	ONLINE_STRUCT_FIELD(FUpdatePresence::Params, Presence)
 END_ONLINE_STRUCT_META()
 
@@ -400,7 +400,7 @@ BEGIN_ONLINE_STRUCT_META(FUpdatePresence::Result)
 END_ONLINE_STRUCT_META()
 
 BEGIN_ONLINE_STRUCT_META(FPartialUpdatePresence::Params)
-	ONLINE_STRUCT_FIELD(FPartialUpdatePresence::Params, LocalUserId),
+	ONLINE_STRUCT_FIELD(FPartialUpdatePresence::Params, LocalAccountId),
 	ONLINE_STRUCT_FIELD(FPartialUpdatePresence::Params, Mutations)
 END_ONLINE_STRUCT_META()
 
@@ -417,7 +417,7 @@ BEGIN_ONLINE_STRUCT_META(FPartialUpdatePresence::Result)
 END_ONLINE_STRUCT_META()
 
 BEGIN_ONLINE_STRUCT_META(FPresenceUpdated)
-	ONLINE_STRUCT_FIELD(FPresenceUpdated, LocalUserId),
+	ONLINE_STRUCT_FIELD(FPresenceUpdated, LocalAccountId),
 	ONLINE_STRUCT_FIELD(FPresenceUpdated, UpdatedPresence)
 END_ONLINE_STRUCT_META()
 
diff --git a/Engine/Plugins/Online/Experimental/OnlineServices/Source/OnlineServicesInterface/Public/Online/Privileges.h b/Engine/Plugins/Online/Experimental/OnlineServices/Source/OnlineServicesInterface/Public/Online/Privileges.h
--- a/Engine/Plugins/Online/Experimental/OnlineServices/Source/OnlineServicesInterface/Public/Online/Privileges.h
+++ b/Engine/Plugins/Online/Experimental/OnlineServices/Source/OnlineServicesInterface/Public/Online/Privileges.h
@@ -65,7 +65,7 @@ struct FQueryUserPrivilege
 	struct Params
 	{
 		/** User id whose privilege we want to query */
-		FAccountId LocalUserId;
+		FAccountId LocalAccountId;
 		/** Privilege to query */
 		EUserPrivileges Privilege;
 	};
@@ -91,7 +91,7 @@ namespace Meta {
 // TODO: Move to OnlineServices_Meta.inl file?
 
 BEGIN_ONLINE_STRUCT_META(FQueryUserPrivilege::Params)
-	ONLINE_STRUCT_FIELD(FQueryUserPrivilege::Params, LocalUserId),
+	ONLINE_STRUCT_FIELD(FQueryUserPrivilege::Params, LocalAccountId),
 	ONLINE_STRUCT_FIELD(FQueryUserPrivilege::Params, Privilege)
 END_ONLINE_STRUCT_META()
 
diff --git a/Engine/Plugins/Online/Experimental/OnlineServices/Source/OnlineServicesInterface/Public/Online/Sessions.h b/Engine/Plugins/Online/Experimental/OnlineServices/Source/OnlineServicesInterface/Public/Online/Sessions.h
--- a/Engine/Plugins/Online/Experimental/OnlineServices/Source/OnlineServicesInterface/Public/Online/Sessions.h
+++ b/Engine/Plugins/Online/Experimental/OnlineServices/Source/OnlineServicesInterface/Public/Online/Sessions.h
@@ -184,7 +184,7 @@ struct ONLINESERVICESINTERFACE_API FSession
 	FSession(const FSession& InSession);
 
 	/** The user who currently owns the session */
-	FAccountId OwnerUserId;
+	FAccountId OwnerAccountId;
 
 	/** The id for the session, platform dependent */
 	FOnlineSessionIdHandle SessionId;
@@ -248,7 +248,7 @@ struct FGetSessionById
 
 	struct Params
 	{
-		FAccountId LocalUserId;
+		FAccountId LocalAccountId;
 
 		FOnlineSessionIdHandle IdHandle;
 	};
@@ -268,7 +268,7 @@ struct FCreateSession
 	struct Params
 	{
 		/** The local user agent which will perform the action. */
-		FAccountId LocalUserId;
+		FAccountId LocalAccountId;
 
 		/** The local name for the session */
 		FName SessionName;
@@ -277,7 +277,7 @@ struct FCreateSession
 		FSessionSettings SessionSettings;
 
 		/** Information for all local users who will join the session (includes the session creator) */
-		FSessionMembersMap LocalUsers;
+		FSessionMembersMap LocalAccounts;
 	};
 
 	struct Result
@@ -293,7 +293,7 @@ struct FUpdateSession
 	struct Params
 	{
 		/** The local user agent which will perform the action. */
-		FAccountId LocalUserId;
+		FAccountId LocalAccountId;
 
 		/** The local name for the session */
 		FName SessionName;
@@ -315,13 +315,13 @@ struct FLeaveSession
 	struct Params
 	{
 		/* The local user agent which leaves the session*/
-		FAccountId LocalUserId;
+		FAccountId LocalAccountId;
 
 		/* The local name for the session. */
 		FName SessionName;
 
 		/** Ids for all local users who will leave the session (includes the main caller) */
-		TArray<FAccountId> LocalUsers;
+		TArray<FAccountId> LocalAccounts;
 
 		/* Whether the call should attempt to destroy the session instead of just leave it */
 		bool bDestroySession;
@@ -340,7 +340,7 @@ struct FFindSessions
 	struct Params
 	{
 		/* The local user agent which starts the session search*/
-		FAccountId LocalUserId;
+		FAccountId LocalAccountId;
 
 		/* Maximum number of results to return in one search */
 		uint32 MaxResults;
@@ -371,10 +371,10 @@ struct FStartMatchmaking
 	struct Params
 	{
 		/** The local user agent which will perform the action. */
-		FAccountId LocalUserId;
+		FAccountId LocalAccountId;
 
 		/* Information for all local users who will join the session (includes the session creator) */
-		FSessionMembersMap LocalUsers;
+		FSessionMembersMap LocalAccounts;
 
 		/* Local name for the session */
 		FName SessionName;
@@ -399,7 +399,7 @@ struct FJoinSession
 	struct Params
 	{
 		/* The local user agent which starts the join operation*/
-		FAccountId LocalUserId;
+		FAccountId LocalAccountId;
 
 		/* Local name for the session */
 		FName SessionName;
@@ -408,7 +408,7 @@ struct FJoinSession
 		FOnlineSessionIdHandle SessionId;
 
 		/* Information for all local users who will join the session (includes the session creator). Any player that wants to join the session needs defined information to become a new member */
-		FSessionMembersMap LocalUsers;
+		FSessionMembersMap LocalAccounts;
 	};
 
 	struct Result
@@ -424,7 +424,7 @@ struct FAddSessionMembers
 	struct Params
 	{
 		/* The local user agent */
-		FAccountId LocalUserId;
+		FAccountId LocalAccountId;
 
 		/* Local name for the session */
 		FName SessionName;
@@ -449,7 +449,7 @@ struct FRemoveSessionMembers
 	struct Params
 	{
 		/* The local user agent */
-		FAccountId LocalUserId;
+		FAccountId LocalAccountId;
 
 		/* Local name for the session */
 		FName SessionName;
@@ -474,7 +474,7 @@ struct FSendSessionInvite
 	struct Params
 	{
 		/* The local user agent which sends the invite*/
-		FAccountId LocalUserId;
+		FAccountId LocalAccountId;
 
 		/* The local name for the session. */
 		FName SessionName;
@@ -496,7 +496,7 @@ struct FGetSessionInvites
 	struct Params
 	{
 		/* The local user agent */
-		FAccountId LocalUserId;
+		FAccountId LocalAccountId;
 	};
 
 	struct Result
@@ -513,7 +513,7 @@ struct FRejectSessionInvite
 	struct Params
 	{
 		/* The local user agent which started the query*/
-		FAccountId LocalUserId;
+		FAccountId LocalAccountId;
 
 		/* The id handle for the invite to be rejected */
 		FOnlineSessionInviteIdHandle SessionInviteId;
@@ -531,7 +531,7 @@ struct FRegisterPlayers
 	struct Params
 	{
 		/* The local user agent */
-		FAccountId LocalUserId;
+		FAccountId LocalAccountId;
 
 		/* The local name for the session. */
 		FName SessionName;
@@ -555,7 +555,7 @@ struct FUnregisterPlayers
 	struct Params
 	{
 		/* The local user agent */
-		FAccountId LocalUserId;
+		FAccountId LocalAccountId;
 
 		/* The local name for the session. */
 		FName SessionName;
@@ -577,7 +577,7 @@ struct FUnregisterPlayers
 struct FSessionJoined
 {
 	/* The local users which joined the session */
-	TArray<FAccountId> LocalUserIds;
+	TArray<FAccountId> LocalAccountIds;
 
 	/* A shared reference to the session joined. */
 	FOnlineSessionIdHandle SessionId;
@@ -586,7 +586,7 @@ struct FSessionJoined
 struct FSessionLeft
 {
 	/* The local users which left the session */
-	TArray<FAccountId> LocalUserIds;
+	TArray<FAccountId> LocalAccountIds;
 };
 
 struct FSessionUpdated
@@ -601,7 +601,7 @@ struct FSessionUpdated
 struct FSessionInviteReceived
 {
 	/* The local user which received the invite */
-	FAccountId LocalUserId;
+	FAccountId LocalAccountId;
 
 	/** The session invite the local user was sent, or the online error if there was a failure retrieving the session for it*/
 	TSharedRef<const FSessionInvite> SessionInvite;
@@ -623,7 +623,7 @@ ONLINESERVICESINTERFACE_API void LexFromString(EUISessionJoinRequestedSource& Ou
 struct FUISessionJoinRequested
 {
 	/** The local user associated with the join request. */
-	FAccountId LocalUserId;
+	FAccountId LocalAccountId;
 
 	/** The id for the session the local user requested to join, or the online error if there was a failure retrieving it */
 	TResult<FOnlineSessionIdHandle, FOnlineError> Result;
@@ -887,7 +887,7 @@ BEGIN_ONLINE_STRUCT_META(FSessionSettingsUpdate)
 END_ONLINE_STRUCT_META()
 
 BEGIN_ONLINE_STRUCT_META(FSession)
-	ONLINE_STRUCT_FIELD(FSession, OwnerUserId),
+	ONLINE_STRUCT_FIELD(FSession, OwnerAccountId),
 	ONLINE_STRUCT_FIELD(FSession, SessionId),
 	ONLINE_STRUCT_FIELD(FSession, SessionSettings)
 END_ONLINE_STRUCT_META()
@@ -915,7 +915,7 @@ BEGIN_ONLINE_STRUCT_META(FGetSessionByName::Result)
 END_ONLINE_STRUCT_META()
 
 BEGIN_ONLINE_STRUCT_META(FGetSessionById::Params)
-	ONLINE_STRUCT_FIELD(FGetSessionById::Params, LocalUserId),
+	ONLINE_STRUCT_FIELD(FGetSessionById::Params, LocalAccountId),
 	ONLINE_STRUCT_FIELD(FGetSessionById::Params, IdHandle)
 END_ONLINE_STRUCT_META()
 
@@ -924,17 +924,17 @@ BEGIN_ONLINE_STRUCT_META(FGetSessionById::Result)
 END_ONLINE_STRUCT_META()
 
 BEGIN_ONLINE_STRUCT_META(FCreateSession::Params)
-	ONLINE_STRUCT_FIELD(FCreateSession::Params, LocalUserId),
+	ONLINE_STRUCT_FIELD(FCreateSession::Params, LocalAccountId),
 	ONLINE_STRUCT_FIELD(FCreateSession::Params, SessionName),
 	ONLINE_STRUCT_FIELD(FCreateSession::Params, SessionSettings),
-	ONLINE_STRUCT_FIELD(FCreateSession::Params, LocalUsers)
+	ONLINE_STRUCT_FIELD(FCreateSession::Params, LocalAccounts)
 END_ONLINE_STRUCT_META()
 
 BEGIN_ONLINE_STRUCT_META(FCreateSession::Result)
 END_ONLINE_STRUCT_META()
 
 BEGIN_ONLINE_STRUCT_META(FUpdateSession::Params)
-	ONLINE_STRUCT_FIELD(FUpdateSession::Params, LocalUserId),
+	ONLINE_STRUCT_FIELD(FUpdateSession::Params, LocalAccountId),
 	ONLINE_STRUCT_FIELD(FUpdateSession::Params, SessionName),
 	ONLINE_STRUCT_FIELD(FUpdateSession::Params, Mutations)
 END_ONLINE_STRUCT_META()
@@ -943,9 +943,9 @@ BEGIN_ONLINE_STRUCT_META(FUpdateSession::Result)
 END_ONLINE_STRUCT_META()
 
 BEGIN_ONLINE_STRUCT_META(FLeaveSession::Params)
-	ONLINE_STRUCT_FIELD(FLeaveSession::Params, LocalUserId),
+	ONLINE_STRUCT_FIELD(FLeaveSession::Params, LocalAccountId),
 	ONLINE_STRUCT_FIELD(FLeaveSession::Params, SessionName),
-	ONLINE_STRUCT_FIELD(FLeaveSession::Params, LocalUsers),
+	ONLINE_STRUCT_FIELD(FLeaveSession::Params, LocalAccounts),
 	ONLINE_STRUCT_FIELD(FLeaveSession::Params, bDestroySession)
 END_ONLINE_STRUCT_META()
 
@@ -953,7 +953,7 @@ BEGIN_ONLINE_STRUCT_META(FLeaveSession::Result)
 END_ONLINE_STRUCT_META()
 
 BEGIN_ONLINE_STRUCT_META(FFindSessions::Params)
-	ONLINE_STRUCT_FIELD(FFindSessions::Params, LocalUserId),
+	ONLINE_STRUCT_FIELD(FFindSessions::Params, LocalAccountId),
 	ONLINE_STRUCT_FIELD(FFindSessions::Params, MaxResults),
 	ONLINE_STRUCT_FIELD(FFindSessions::Params, bFindLANSessions),	
 	ONLINE_STRUCT_FIELD(FFindSessions::Params, Filters),
@@ -966,8 +966,8 @@ BEGIN_ONLINE_STRUCT_META(FFindSessions::Result)
 END_ONLINE_STRUCT_META()
 
 BEGIN_ONLINE_STRUCT_META(FStartMatchmaking::Params)
-	ONLINE_STRUCT_FIELD(FStartMatchmaking::Params, LocalUserId),
-	ONLINE_STRUCT_FIELD(FStartMatchmaking::Params, LocalUsers),
+	ONLINE_STRUCT_FIELD(FStartMatchmaking::Params, LocalAccountId),
+	ONLINE_STRUCT_FIELD(FStartMatchmaking::Params, LocalAccounts),
 	ONLINE_STRUCT_FIELD(FStartMatchmaking::Params, SessionName),
 	ONLINE_STRUCT_FIELD(FStartMatchmaking::Params, SessionSettings),
 	ONLINE_STRUCT_FIELD(FStartMatchmaking::Params, SearchFilters)
@@ -977,17 +977,17 @@ BEGIN_ONLINE_STRUCT_META(FStartMatchmaking::Result)
 END_ONLINE_STRUCT_META()
 
 BEGIN_ONLINE_STRUCT_META(FJoinSession::Params)
-	ONLINE_STRUCT_FIELD(FJoinSession::Params, LocalUserId),
+	ONLINE_STRUCT_FIELD(FJoinSession::Params, LocalAccountId),
 	ONLINE_STRUCT_FIELD(FJoinSession::Params, SessionName),
 	ONLINE_STRUCT_FIELD(FJoinSession::Params, SessionId),
-	ONLINE_STRUCT_FIELD(FJoinSession::Params, LocalUsers)
+	ONLINE_STRUCT_FIELD(FJoinSession::Params, LocalAccounts)
 END_ONLINE_STRUCT_META()
 
 BEGIN_ONLINE_STRUCT_META(FJoinSession::Result)
 END_ONLINE_STRUCT_META()
 
 BEGIN_ONLINE_STRUCT_META(FAddSessionMembers::Params)
-	ONLINE_STRUCT_FIELD(FAddSessionMembers::Params, LocalUserId),
+	ONLINE_STRUCT_FIELD(FAddSessionMembers::Params, LocalAccountId),
 	ONLINE_STRUCT_FIELD(FAddSessionMembers::Params, SessionName),
 	ONLINE_STRUCT_FIELD(FAddSessionMembers::Params, NewSessionMembers),
 	ONLINE_STRUCT_FIELD(FAddSessionMembers::Params, bRegisterPlayers)
@@ -997,7 +997,7 @@ BEGIN_ONLINE_STRUCT_META(FAddSessionMembers::Result)
 END_ONLINE_STRUCT_META()
 
 BEGIN_ONLINE_STRUCT_META(FRemoveSessionMembers::Params)
-	ONLINE_STRUCT_FIELD(FRemoveSessionMembers::Params, LocalUserId),
+	ONLINE_STRUCT_FIELD(FRemoveSessionMembers::Params, LocalAccountId),
 	ONLINE_STRUCT_FIELD(FRemoveSessionMembers::Params, SessionName),
 	ONLINE_STRUCT_FIELD(FRemoveSessionMembers::Params, SessionMemberIds),
 	ONLINE_STRUCT_FIELD(FRemoveSessionMembers::Params, bUnregisterPlayers)
@@ -1007,7 +1007,7 @@ BEGIN_ONLINE_STRUCT_META(FRemoveSessionMembers::Result)
 END_ONLINE_STRUCT_META()
 
 BEGIN_ONLINE_STRUCT_META(FSendSessionInvite::Params)
-	ONLINE_STRUCT_FIELD(FSendSessionInvite::Params, LocalUserId),
+	ONLINE_STRUCT_FIELD(FSendSessionInvite::Params, LocalAccountId),
 	ONLINE_STRUCT_FIELD(FSendSessionInvite::Params, SessionName),
 	ONLINE_STRUCT_FIELD(FSendSessionInvite::Params, TargetUsers)
 END_ONLINE_STRUCT_META()
@@ -1016,7 +1016,7 @@ BEGIN_ONLINE_STRUCT_META(FSendSessionInvite::Result)
 END_ONLINE_STRUCT_META()
 
 BEGIN_ONLINE_STRUCT_META(FGetSessionInvites::Params)
-	ONLINE_STRUCT_FIELD(FGetSessionInvites::Params, LocalUserId)
+	ONLINE_STRUCT_FIELD(FGetSessionInvites::Params, LocalAccountId)
 END_ONLINE_STRUCT_META()
 
 BEGIN_ONLINE_STRUCT_META(FGetSessionInvites::Result)
@@ -1024,7 +1024,7 @@ BEGIN_ONLINE_STRUCT_META(FGetSessionInvites::Result)
 END_ONLINE_STRUCT_META()
 
 BEGIN_ONLINE_STRUCT_META(FRejectSessionInvite::Params)
-	ONLINE_STRUCT_FIELD(FRejectSessionInvite::Params, LocalUserId),
+	ONLINE_STRUCT_FIELD(FRejectSessionInvite::Params, LocalAccountId),
 	ONLINE_STRUCT_FIELD(FRejectSessionInvite::Params, SessionInviteId)
 END_ONLINE_STRUCT_META()
 
@@ -1032,7 +1032,7 @@ BEGIN_ONLINE_STRUCT_META(FRejectSessionInvite::Result)
 END_ONLINE_STRUCT_META()
 
 BEGIN_ONLINE_STRUCT_META(FRegisterPlayers::Params)
-	ONLINE_STRUCT_FIELD(FRegisterPlayers::Params, LocalUserId),
+	ONLINE_STRUCT_FIELD(FRegisterPlayers::Params, LocalAccountId),
 	ONLINE_STRUCT_FIELD(FRegisterPlayers::Params, SessionName),
 	ONLINE_STRUCT_FIELD(FRegisterPlayers::Params, TargetUsers),
 	ONLINE_STRUCT_FIELD(FRegisterPlayers::Params, bReserveSlot)
@@ -1042,7 +1042,7 @@ BEGIN_ONLINE_STRUCT_META(FRegisterPlayers::Result)
 END_ONLINE_STRUCT_META()
 
 BEGIN_ONLINE_STRUCT_META(FUnregisterPlayers::Params)
-	ONLINE_STRUCT_FIELD(FUnregisterPlayers::Params, LocalUserId),
+	ONLINE_STRUCT_FIELD(FUnregisterPlayers::Params, LocalAccountId),
 	ONLINE_STRUCT_FIELD(FUnregisterPlayers::Params, SessionName),
 	ONLINE_STRUCT_FIELD(FUnregisterPlayers::Params, TargetUsers),
 	ONLINE_STRUCT_FIELD(FUnregisterPlayers::Params, bRemoveUnregisteredPlayers)
diff --git a/Engine/Plugins/Online/Experimental/OnlineServices/Source/OnlineServicesInterface/Public/Online/Stats.h b/Engine/Plugins/Online/Experimental/OnlineServices/Source/OnlineServicesInterface/Public/Online/Stats.h
--- a/Engine/Plugins/Online/Experimental/OnlineServices/Source/OnlineServicesInterface/Public/Online/Stats.h
+++ b/Engine/Plugins/Online/Experimental/OnlineServices/Source/OnlineServicesInterface/Public/Online/Stats.h
@@ -14,7 +14,7 @@ using FStatValue = FSchemaVariant;
 struct FUserStats
 {
 	/* The user id */
-	FAccountId UserId;
+	FAccountId AccountId;
 	/* The stats of one account */
 	TMap<FString, FStatValue> Stats;
 };
@@ -26,7 +26,7 @@ struct FUpdateStats
 	struct Params
 	{
 		/* Local user id */
-		FAccountId LocalUserId;
+		FAccountId LocalAccountId;
 		/* The stats of users to update */
 		TArray<FUserStats> UpdateUsersStats;
 	};
@@ -43,9 +43,9 @@ struct FQueryStats
 	struct Params
 	{
 		/* Local user id */
-		FAccountId LocalUserId;
+		FAccountId LocalAccountId;
 		/* Target user id to query */
-		FAccountId TargetUserId;
+		FAccountId TargetAccountId;
 		/* The stats to query */
 		TArray<FString> StatNames;
 	};
@@ -64,9 +64,9 @@ struct FBatchQueryStats
 	struct Params
 	{
 		/* Local user id */
-		FAccountId LocalUserId;
+		FAccountId LocalAccountId;
 		/* The user ids to query */
-		TArray<FAccountId> TargetUserIds;
+		TArray<FAccountId> TargetAccountIds;
 		/* The stats to query */
 		TArray<FString> StatNames;
 	};
@@ -86,7 +86,7 @@ struct FResetStats
 	struct Params
 	{
 		/* Local user id */
-		FAccountId LocalUserId;
+		FAccountId LocalAccountId;
 	};
 
 	struct Result
@@ -163,12 +163,12 @@ public:
 namespace Meta {
 
 BEGIN_ONLINE_STRUCT_META(FUserStats)
-	ONLINE_STRUCT_FIELD(FUserStats, UserId),
+	ONLINE_STRUCT_FIELD(FUserStats, AccountId),
 	ONLINE_STRUCT_FIELD(FUserStats, Stats)
 END_ONLINE_STRUCT_META()
 
 BEGIN_ONLINE_STRUCT_META(FUpdateStats::Params)
-	ONLINE_STRUCT_FIELD(FUpdateStats::Params, LocalUserId),
+	ONLINE_STRUCT_FIELD(FUpdateStats::Params, LocalAccountId),
 	ONLINE_STRUCT_FIELD(FUpdateStats::Params, UpdateUsersStats)
 END_ONLINE_STRUCT_META()
 
@@ -176,8 +176,8 @@ BEGIN_ONLINE_STRUCT_META(FUpdateStats::Result)
 END_ONLINE_STRUCT_META()
 
 BEGIN_ONLINE_STRUCT_META(FQueryStats::Params)
-	ONLINE_STRUCT_FIELD(FQueryStats::Params, LocalUserId),
-	ONLINE_STRUCT_FIELD(FQueryStats::Params, TargetUserId)
+	ONLINE_STRUCT_FIELD(FQueryStats::Params, LocalAccountId),
+	ONLINE_STRUCT_FIELD(FQueryStats::Params, TargetAccountId)
 END_ONLINE_STRUCT_META()
 
 BEGIN_ONLINE_STRUCT_META(FQueryStats::Result)
@@ -185,8 +185,8 @@ BEGIN_ONLINE_STRUCT_META(FQueryStats::Result)
 END_ONLINE_STRUCT_META()
 
 BEGIN_ONLINE_STRUCT_META(FBatchQueryStats::Params)
-	ONLINE_STRUCT_FIELD(FBatchQueryStats::Params, LocalUserId),
-	ONLINE_STRUCT_FIELD(FBatchQueryStats::Params, TargetUserIds)
+	ONLINE_STRUCT_FIELD(FBatchQueryStats::Params, LocalAccountId),
+	ONLINE_STRUCT_FIELD(FBatchQueryStats::Params, TargetAccountIds)
 END_ONLINE_STRUCT_META()
 
 BEGIN_ONLINE_STRUCT_META(FBatchQueryStats::Result)
@@ -195,7 +195,7 @@ END_ONLINE_STRUCT_META()
 
 #if !UE_BUILD_SHIPPING
 BEGIN_ONLINE_STRUCT_META(FResetStats::Params)
-	ONLINE_STRUCT_FIELD(FResetStats::Params, LocalUserId)
+	ONLINE_STRUCT_FIELD(FResetStats::Params, LocalAccountId)
 END_ONLINE_STRUCT_META()
 
 BEGIN_ONLINE_STRUCT_META(FResetStats::Result)
diff --git a/Engine/Plugins/Online/Experimental/OnlineServices/Source/OnlineServicesInterface/Public/Online/TitleFile.h b/Engine/Plugins/Online/Experimental/OnlineServices/Source/OnlineServicesInterface/Public/Online/TitleFile.h
--- a/Engine/Plugins/Online/Experimental/OnlineServices/Source/OnlineServicesInterface/Public/Online/TitleFile.h
+++ b/Engine/Plugins/Online/Experimental/OnlineServices/Source/OnlineServicesInterface/Public/Online/TitleFile.h
@@ -15,7 +15,7 @@ struct FTitleFileEnumerateFiles
 	struct Params
 	{
 		/** Local user performing the operation */
-		FAccountId LocalUserId;
+		FAccountId LocalAccountId;
 	};
 
 	struct Result
@@ -30,7 +30,7 @@ struct FTitleFileGetEnumeratedFiles
 	struct Params
 	{
 		/** Local user performing the operation */
-		FAccountId LocalUserId;
+		FAccountId LocalAccountId;
 	};
 
 	struct Result
@@ -52,7 +52,7 @@ struct FTitleFileReadFile
 	struct Params
 	{
 		/** Local user performing the operation */
-		FAccountId LocalUserId;
+		FAccountId LocalAccountId;
 		/** File to be opened */
 		FString Filename;
 	};
@@ -87,14 +87,14 @@ public:
 namespace Meta {
 
 BEGIN_ONLINE_STRUCT_META(FTitleFileEnumerateFiles::Params)
-	ONLINE_STRUCT_FIELD(FTitleFileEnumerateFiles::Params, LocalUserId)
+	ONLINE_STRUCT_FIELD(FTitleFileEnumerateFiles::Params, LocalAccountId)
 END_ONLINE_STRUCT_META()
 
 BEGIN_ONLINE_STRUCT_META(FTitleFileEnumerateFiles::Result)
 END_ONLINE_STRUCT_META()
 
 BEGIN_ONLINE_STRUCT_META(FTitleFileGetEnumeratedFiles::Params)
-	ONLINE_STRUCT_FIELD(FTitleFileGetEnumeratedFiles::Params, LocalUserId)
+	ONLINE_STRUCT_FIELD(FTitleFileGetEnumeratedFiles::Params, LocalAccountId)
 END_ONLINE_STRUCT_META()
 
 BEGIN_ONLINE_STRUCT_META(FTitleFileGetEnumeratedFiles::Result)
@@ -102,7 +102,7 @@ BEGIN_ONLINE_STRUCT_META(FTitleFileGetEnumeratedFiles::Result)
 END_ONLINE_STRUCT_META()
 
 BEGIN_ONLINE_STRUCT_META(FTitleFileReadFile::Params)
-	ONLINE_STRUCT_FIELD(FTitleFileReadFile::Params, LocalUserId),
+	ONLINE_STRUCT_FIELD(FTitleFileReadFile::Params, LocalAccountId),
 	ONLINE_STRUCT_FIELD(FTitleFileReadFile::Params, Filename)
 END_ONLINE_STRUCT_META()
 
diff --git a/Engine/Plugins/Online/Experimental/OnlineServices/Source/OnlineServicesInterface/Public/Online/UserFile.h b/Engine/Plugins/Online/Experimental/OnlineServices/Source/OnlineServicesInterface/Public/Online/UserFile.h
--- a/Engine/Plugins/Online/Experimental/OnlineServices/Source/OnlineServicesInterface/Public/Online/UserFile.h
+++ b/Engine/Plugins/Online/Experimental/OnlineServices/Source/OnlineServicesInterface/Public/Online/UserFile.h
@@ -15,7 +15,7 @@ struct FUserFileEnumerateFiles
 	struct Params
 	{
 		/** Local user performing the operation */
-		FAccountId LocalUserId;
+		FAccountId LocalAccountId;
 	};
 
 	struct Result
@@ -30,7 +30,7 @@ struct FUserFileGetEnumeratedFiles
 	struct Params
 	{
 		/** Local user performing the operation */
-		FAccountId LocalUserId;
+		FAccountId LocalAccountId;
 	};
 
 	struct Result
@@ -52,7 +52,7 @@ struct FUserFileReadFile
 	struct Params
 	{
 		/** Local user performing the operation */
-		FAccountId LocalUserId;
+		FAccountId LocalAccountId;
 		/** File to be read */
 		FString Filename;
 	};
@@ -71,7 +71,7 @@ struct FUserFileWriteFile
 	struct Params
 	{
 		/** Local user performing the operation */
-		FAccountId LocalUserId;
+		FAccountId LocalAccountId;
 		/** File to be written */
 		FString Filename;
 		/** Data to write */
@@ -90,7 +90,7 @@ struct FUserFileCopyFile
 	struct Params
 	{
 		/** Local user performing the operation */
-		FAccountId LocalUserId;
+		FAccountId LocalAccountId;
 		/** File to copy contents from */
 		FString SourceFilename;
 		/** File to copy contents to */
@@ -109,7 +109,7 @@ struct FUserFileDeleteFile
 	struct Params
 	{
 		/** Local user performing the operation */
-		FAccountId LocalUserId;
+		FAccountId LocalAccountId;
 		/** File to delete */
 		FString Filename;
 	};
@@ -156,14 +156,14 @@ public:
 namespace Meta {
 
 BEGIN_ONLINE_STRUCT_META(FUserFileEnumerateFiles::Params)
-	ONLINE_STRUCT_FIELD(FUserFileEnumerateFiles::Params, LocalUserId)
+	ONLINE_STRUCT_FIELD(FUserFileEnumerateFiles::Params, LocalAccountId)
 END_ONLINE_STRUCT_META()
 
 BEGIN_ONLINE_STRUCT_META(FUserFileEnumerateFiles::Result)
 END_ONLINE_STRUCT_META()
 
 BEGIN_ONLINE_STRUCT_META(FUserFileGetEnumeratedFiles::Params)
-	ONLINE_STRUCT_FIELD(FUserFileGetEnumeratedFiles::Params, LocalUserId)
+	ONLINE_STRUCT_FIELD(FUserFileGetEnumeratedFiles::Params, LocalAccountId)
 END_ONLINE_STRUCT_META()
 
 BEGIN_ONLINE_STRUCT_META(FUserFileGetEnumeratedFiles::Result)
@@ -171,7 +171,7 @@ BEGIN_ONLINE_STRUCT_META(FUserFileGetEnumeratedFiles::Result)
 END_ONLINE_STRUCT_META()
 
 BEGIN_ONLINE_STRUCT_META(FUserFileReadFile::Params)
-	ONLINE_STRUCT_FIELD(FUserFileReadFile::Params, LocalUserId),
+	ONLINE_STRUCT_FIELD(FUserFileReadFile::Params, LocalAccountId),
 	ONLINE_STRUCT_FIELD(FUserFileReadFile::Params, Filename)
 END_ONLINE_STRUCT_META()
 
@@ -180,7 +180,7 @@ BEGIN_ONLINE_STRUCT_META(FUserFileReadFile::Result)
 END_ONLINE_STRUCT_META()
 
 BEGIN_ONLINE_STRUCT_META(FUserFileWriteFile::Params)
-	ONLINE_STRUCT_FIELD(FUserFileWriteFile::Params, LocalUserId),
+	ONLINE_STRUCT_FIELD(FUserFileWriteFile::Params, LocalAccountId),
 	ONLINE_STRUCT_FIELD(FUserFileWriteFile::Params, Filename),
 	ONLINE_STRUCT_FIELD(FUserFileWriteFile::Params, FileContents)
 END_ONLINE_STRUCT_META()
@@ -189,7 +189,7 @@ BEGIN_ONLINE_STRUCT_META(FUserFileWriteFile::Result)
 END_ONLINE_STRUCT_META()
 
 BEGIN_ONLINE_STRUCT_META(FUserFileCopyFile::Params)
-	ONLINE_STRUCT_FIELD(FUserFileCopyFile::Params, LocalUserId),
+	ONLINE_STRUCT_FIELD(FUserFileCopyFile::Params, LocalAccountId),
 	ONLINE_STRUCT_FIELD(FUserFileCopyFile::Params, SourceFilename),
 	ONLINE_STRUCT_FIELD(FUserFileCopyFile::Params, TargetFilename)
 END_ONLINE_STRUCT_META()
@@ -198,7 +198,7 @@ BEGIN_ONLINE_STRUCT_META(FUserFileCopyFile::Result)
 END_ONLINE_STRUCT_META()
 
 BEGIN_ONLINE_STRUCT_META(FUserFileDeleteFile::Params)
-	ONLINE_STRUCT_FIELD(FUserFileDeleteFile::Params, LocalUserId),
+	ONLINE_STRUCT_FIELD(FUserFileDeleteFile::Params, LocalAccountId),
 	ONLINE_STRUCT_FIELD(FUserFileDeleteFile::Params, Filename)
 END_ONLINE_STRUCT_META()
 
diff --git a/Engine/Plugins/Online/Experimental/OnlineServices/Source/OnlineServicesInterface/Public/Online/UserInfo.h b/Engine/Plugins/Online/Experimental/OnlineServices/Source/OnlineServicesInterface/Public/Online/UserInfo.h
--- a/Engine/Plugins/Online/Experimental/OnlineServices/Source/OnlineServicesInterface/Public/Online/UserInfo.h
+++ b/Engine/Plugins/Online/Experimental/OnlineServices/Source/OnlineServicesInterface/Public/Online/UserInfo.h
@@ -11,7 +11,7 @@ namespace UE::Online {
 struct FUserInfo
 {
 	/* Account Id */
-	FAccountId UserId;
+	FAccountId AccountId;
 	/* Display Name of the User */
 	FString DisplayName;
 };
@@ -24,9 +24,9 @@ struct FQueryUserInfo
 	struct Params
 	{
 		/* Local user to query users for */
-		FAccountId LocalUserId;
+		FAccountId LocalAccountId;
 		/* List of User ids to Query*/
-		TArray<FAccountId> UserIds;
+		TArray<FAccountId> AccountIds;
 	};
 
 	/**
@@ -46,9 +46,9 @@ struct FGetUserInfo
 	struct Params
 	{
 		/* Local user */
-		FAccountId LocalUserId;
+		FAccountId LocalAccountId;
 		/* User id to get info for*/
-		FAccountId UserId;
+		FAccountId AccountId;
 	};
 
 	/**
@@ -69,9 +69,9 @@ struct FQueryUserAvatar
 	struct Params
 	{
 		/* Local user to query users for */
-		FAccountId LocalUserId;
+		FAccountId LocalAccountId;
 		/* List of User ids to Query*/
-		TArray<FAccountId> UserIds;
+		TArray<FAccountId> AccountIds;
 	};
 
 	/**
@@ -91,9 +91,9 @@ struct FGetUserAvatar
 	struct Params
 	{
 		/* Local user */
-		FAccountId LocalUserId;
+		FAccountId LocalAccountId;
 		/* User id to get avatar for*/
-		FAccountId UserId;
+		FAccountId AccountId;
 	};
 
 	/**
@@ -114,9 +114,9 @@ struct FShowUserProfile
 	struct Params
 	{
 		/* Local user */
-		FAccountId LocalUserId;
+		FAccountId LocalAccountId;
 		/* User id to display the profile for*/
-		FAccountId UserId;
+		FAccountId AccountId;
 	};
 
 	/**
@@ -182,21 +182,21 @@ public:
 namespace Meta {
 
 BEGIN_ONLINE_STRUCT_META(FUserInfo)
-	ONLINE_STRUCT_FIELD(FUserInfo, UserId),
+	ONLINE_STRUCT_FIELD(FUserInfo, AccountId),
 	ONLINE_STRUCT_FIELD(FUserInfo, DisplayName)
 END_ONLINE_STRUCT_META()
 
 BEGIN_ONLINE_STRUCT_META(FQueryUserInfo::Params)
-	ONLINE_STRUCT_FIELD(FQueryUserInfo::Params, LocalUserId),
-	ONLINE_STRUCT_FIELD(FQueryUserInfo::Params, UserIds)
+	ONLINE_STRUCT_FIELD(FQueryUserInfo::Params, LocalAccountId),
+	ONLINE_STRUCT_FIELD(FQueryUserInfo::Params, AccountIds)
 END_ONLINE_STRUCT_META()
 
 BEGIN_ONLINE_STRUCT_META(FQueryUserInfo::Result)
 END_ONLINE_STRUCT_META()
 
 BEGIN_ONLINE_STRUCT_META(FGetUserInfo::Params)
-	ONLINE_STRUCT_FIELD(FGetUserInfo::Params, LocalUserId),
-	ONLINE_STRUCT_FIELD(FGetUserInfo::Params, UserId)
+	ONLINE_STRUCT_FIELD(FGetUserInfo::Params, LocalAccountId),
+	ONLINE_STRUCT_FIELD(FGetUserInfo::Params, AccountId)
 END_ONLINE_STRUCT_META()
 
 BEGIN_ONLINE_STRUCT_META(FGetUserInfo::Result)
@@ -204,16 +204,16 @@ BEGIN_ONLINE_STRUCT_META(FGetUserInfo::Result)
 END_ONLINE_STRUCT_META()
 
 BEGIN_ONLINE_STRUCT_META(FQueryUserAvatar::Params)
-	ONLINE_STRUCT_FIELD(FQueryUserAvatar::Params, LocalUserId),
-	ONLINE_STRUCT_FIELD(FQueryUserAvatar::Params, UserIds)
+	ONLINE_STRUCT_FIELD(FQueryUserAvatar::Params, LocalAccountId),
+	ONLINE_STRUCT_FIELD(FQueryUserAvatar::Params, AccountIds)
 END_ONLINE_STRUCT_META()
 
 BEGIN_ONLINE_STRUCT_META(FQueryUserAvatar::Result)
 END_ONLINE_STRUCT_META()
 
 BEGIN_ONLINE_STRUCT_META(FGetUserAvatar::Params)
-	ONLINE_STRUCT_FIELD(FGetUserAvatar::Params, LocalUserId),
-	ONLINE_STRUCT_FIELD(FGetUserAvatar::Params, UserId)
+	ONLINE_STRUCT_FIELD(FGetUserAvatar::Params, LocalAccountId),
+	ONLINE_STRUCT_FIELD(FGetUserAvatar::Params, AccountId)
 END_ONLINE_STRUCT_META()
 
 BEGIN_ONLINE_STRUCT_META(FGetUserAvatar::Result)
@@ -221,8 +221,8 @@ BEGIN_ONLINE_STRUCT_META(FGetUserAvatar::Result)
 END_ONLINE_STRUCT_META()
 
 BEGIN_ONLINE_STRUCT_META(FShowUserProfile::Params)
-	ONLINE_STRUCT_FIELD(FShowUserProfile::Params, LocalUserId),
-	ONLINE_STRUCT_FIELD(FShowUserProfile::Params, UserId)
+	ONLINE_STRUCT_FIELD(FShowUserProfile::Params, LocalAccountId),
+	ONLINE_STRUCT_FIELD(FShowUserProfile::Params, AccountId)
 END_ONLINE_STRUCT_META()
 
 BEGIN_ONLINE_STRUCT_META(FShowUserProfile::Result)
diff --git a/Engine/Plugins/Online/Experimental/OnlineServicesEOS/Source/Private/Online/AuthEOS.cpp b/Engine/Plugins/Online/Experimental/OnlineServicesEOS/Source/Private/Online/AuthEOS.cpp
--- a/Engine/Plugins/Online/Experimental/OnlineServicesEOS/Source/Private/Online/AuthEOS.cpp
+++ b/Engine/Plugins/Online/Experimental/OnlineServicesEOS/Source/Private/Online/AuthEOS.cpp
@@ -257,12 +257,12 @@ TOnlineAsyncOpHandle<FAuthLogin> FAuthEOS::Login(FAuthLogin::Params&& Params)
 }
 
 template <typename IdType>
-TFuture<FAccountId> ResolveAccountIdImpl(FAuthEOS& AuthEOS, const FAccountId& LocalUserId, const IdType InId)
+TFuture<FAccountId> ResolveAccountIdImpl(FAuthEOS& AuthEOS, const FAccountId& LocalAccountId, const IdType InId)
 {
 	TPromise<FAccountId> Promise;
 	TFuture<FAccountId> Future = Promise.GetFuture();
 
-	AuthEOS.ResolveAccountIds(LocalUserId, { InId }).Next([Promise = MoveTemp(Promise)](const TArray<FAccountId>& AccountIds) mutable
+	AuthEOS.ResolveAccountIds(LocalAccountId, { InId }).Next([Promise = MoveTemp(Promise)](const TArray<FAccountId>& AccountIds) mutable
 	{
 		FAccountId Result;
 		if (AccountIds.Num() == 1)
@@ -282,7 +282,7 @@ TOnlineAsyncOpHandle<FAuthQueryExternalServerAuthTicket> FAuthEOS::QueryExternal
 	Op->Then([this](TOnlineAsyncOp<FAuthQueryExternalServerAuthTicket>& InAsyncOp)
 	{
 		const FAuthQueryExternalServerAuthTicket::Params& Params = InAsyncOp.GetParams();
-		TSharedPtr<FAccountInfoEOS> AccountInfoEOS = AccountInfoRegistryEOS.Find(Params.LocalUserId);
+		TSharedPtr<FAccountInfoEOS> AccountInfoEOS = AccountInfoRegistryEOS.Find(Params.LocalAccountId);
 		if (!AccountInfoEOS)
 		{
 			InAsyncOp.SetError(Errors::InvalidParams());
@@ -325,7 +325,7 @@ TOnlineAsyncOpHandle<FAuthQueryExternalAuthToken> FAuthEOS::QueryExternalAuthTok
 	Op->Then([this](TOnlineAsyncOp<FAuthQueryExternalAuthToken>& InAsyncOp)
 	{
 		const FAuthQueryExternalAuthToken::Params& Params = InAsyncOp.GetParams();
-		TSharedPtr<FAccountInfoEOS> AccountInfoEOS = AccountInfoRegistryEOS.Find(Params.LocalUserId);
+		TSharedPtr<FAccountInfoEOS> AccountInfoEOS = AccountInfoRegistryEOS.Find(Params.LocalAccountId);
 		if (!AccountInfoEOS)
 		{
 			InAsyncOp.SetError(Errors::InvalidParams());
@@ -347,19 +347,19 @@ TOnlineAsyncOpHandle<FAuthQueryExternalAuthToken> FAuthEOS::QueryExternalAuthTok
 	return Op->GetHandle();
 }
 
-TFuture<FAccountId> FAuthEOS::ResolveAccountId(const FAccountId& LocalUserId, const EOS_EpicAccountId EpicAccountId)
+TFuture<FAccountId> FAuthEOS::ResolveAccountId(const FAccountId& LocalAccountId, const EOS_EpicAccountId EpicAccountId)
 {
-	return ResolveAccountIdImpl(*this, LocalUserId, EpicAccountId);
+	return ResolveAccountIdImpl(*this, LocalAccountId, EpicAccountId);
 }
 
-TFuture<FAccountId> FAuthEOS::ResolveAccountId(const FAccountId& LocalUserId, const EOS_ProductUserId ProductUserId)
+TFuture<FAccountId> FAuthEOS::ResolveAccountId(const FAccountId& LocalAccountId, const EOS_ProductUserId ProductUserId)
 {
-	return ResolveAccountIdImpl(*this, LocalUserId, ProductUserId);
+	return ResolveAccountIdImpl(*this, LocalAccountId, ProductUserId);
 }
 
 using FEpicAccountIdStrBuffer = char[EOS_EPICACCOUNTID_MAX_LENGTH + 1];
 
-TFuture<TArray<FAccountId>> FAuthEOS::ResolveAccountIds(const FAccountId& LocalUserId, const TArray<EOS_EpicAccountId>& InEpicAccountIds)
+TFuture<TArray<FAccountId>> FAuthEOS::ResolveAccountIds(const FAccountId& LocalAccountId, const TArray<EOS_EpicAccountId>& InEpicAccountIds)
 {
 	// Search for all the account id's
 	TArray<FAccountId> AccountIdHandles;
@@ -386,8 +386,8 @@ TFuture<TArray<FAccountId>> FAuthEOS::ResolveAccountIds(const FAccountId& LocalU
 		return MakeFulfilledPromise<TArray<FAccountId>>(MoveTemp(AccountIdHandles)).GetFuture();
 	}
 
-	// If we failed to find all the handles, we need to query, which requires a valid LocalUserId
-	if (!ValidateOnlineId(LocalUserId))
+	// If we failed to find all the handles, we need to query, which requires a valid LocalAccountId
+	if (!ValidateOnlineId(LocalAccountId))
 	{
 		checkNoEntry();
 		return MakeFulfilledPromise<TArray<FAccountId>>().GetFuture();
@@ -415,7 +415,7 @@ TFuture<TArray<FAccountId>> FAuthEOS::ResolveAccountIds(const FAccountId& LocalU
 
 	EOS_Connect_QueryExternalAccountMappingsOptions Options = {};
 	Options.ApiVersion = EOS_CONNECT_QUERYEXTERNALACCOUNTMAPPINGS_API_LATEST;
-	Options.LocalUserId = GetProductUserIdChecked(LocalUserId);
+	Options.LocalUserId = GetProductUserIdChecked(LocalAccountId);
 	Options.AccountIdType = EOS_EExternalAccountType::EOS_EAT_EPIC;
 	Options.ExternalAccountIds = (const char**)EpicAccountIdStrPtrs.GetData();
 	Options.ExternalAccountIdCount = 1;
@@ -462,7 +462,7 @@ TFuture<TArray<FAccountId>> FAuthEOS::ResolveAccountIds(const FAccountId& LocalU
 	return Future;
 }
 
-TFuture<TArray<FAccountId>> FAuthEOS::ResolveAccountIds(const FAccountId& LocalUserId, const TArray<EOS_ProductUserId>& InProductUserIds)
+TFuture<TArray<FAccountId>> FAuthEOS::ResolveAccountIds(const FAccountId& LocalAccountId, const TArray<EOS_ProductUserId>& InProductUserIds)
 {
 	// Search for all the account id's
 	TArray<FAccountId> AccountIdHandles;
@@ -489,8 +489,8 @@ TFuture<TArray<FAccountId>> FAuthEOS::ResolveAccountIds(const FAccountId& LocalU
 		return MakeFulfilledPromise<TArray<FAccountId>>(MoveTemp(AccountIdHandles)).GetFuture();
 	}
 
-	// If we failed to find all the handles, we need to query, which requires a valid LocalUserId
-	if (!ValidateOnlineId(LocalUserId))
+	// If we failed to find all the handles, we need to query, which requires a valid LocalAccountId
+	if (!ValidateOnlineId(LocalAccountId))
 	{
 		checkNoEntry();
 		return MakeFulfilledPromise<TArray<FAccountId>>().GetFuture();
@@ -501,7 +501,7 @@ TFuture<TArray<FAccountId>> FAuthEOS::ResolveAccountIds(const FAccountId& LocalU
 
 	EOS_Connect_QueryProductUserIdMappingsOptions Options = {};
 	Options.ApiVersion = EOS_CONNECT_QUERYPRODUCTUSERIDMAPPINGS_API_LATEST;
-	Options.LocalUserId = GetProductUserIdChecked(LocalUserId);
+	Options.LocalUserId = GetProductUserIdChecked(LocalAccountId);
 	Options.ProductUserIds = MissingProductUserIds.GetData();
 	Options.ProductUserIdCount = MissingProductUserIds.Num();
 	static_assert(EOS_CONNECT_QUERYPRODUCTUSERIDMAPPINGS_API_LATEST == 2, "EOS_Connect_QueryProductUserIdMappingsOptions updated, check new fields");
@@ -556,12 +556,12 @@ TFunction<TFuture<FAccountId>(FOnlineAsyncOp& InAsyncOp, const ParamType& Param)
 {
 	return [AuthEOS](FOnlineAsyncOp& InAsyncOp, const ParamType& Param)
 	{
-		const FAccountId* LocalUserIdPtr = InAsyncOp.Data.Get<FAccountId>(TEXT("LocalUserId"));
-		if (!ensure(LocalUserIdPtr))
+		const FAccountId* LocalAccountIdPtr = InAsyncOp.Data.Get<FAccountId>(TEXT("LocalAccountId"));
+		if (!ensure(LocalAccountIdPtr))
 		{
 			return MakeFulfilledPromise<FAccountId>().GetFuture();
 		}
-		return AuthEOS->ResolveAccountId(*LocalUserIdPtr, Param);
+		return AuthEOS->ResolveAccountId(*LocalAccountIdPtr, Param);
 	};
 }
 TFunction<TFuture<FAccountId>(FOnlineAsyncOp& InAsyncOp, const EOS_EpicAccountId& ProductUserId)> FAuthEOS::ResolveEpicIdFn()
@@ -578,12 +578,12 @@ TFunction<TFuture<TArray<FAccountId>>(FOnlineAsyncOp& InAsyncOp, const TArray<Pa
 {
 	return [AuthEOS](FOnlineAsyncOp& InAsyncOp, const TArray<ParamType>& Param)
 	{
-		const FAccountId* LocalUserIdPtr = InAsyncOp.Data.Get<FAccountId>(TEXT("LocalUserId"));
-		if (!ensure(LocalUserIdPtr))
+		const FAccountId* LocalAccountIdPtr = InAsyncOp.Data.Get<FAccountId>(TEXT("LocalAccountId"));
+		if (!ensure(LocalAccountIdPtr))
 		{
 			return MakeFulfilledPromise<TArray<FAccountId>>().GetFuture();
 		}
-		return AuthEOS->ResolveAccountIds(*LocalUserIdPtr, Param);
+		return AuthEOS->ResolveAccountIds(*LocalAccountIdPtr, Param);
 	};
 }
 TFunction<TFuture<TArray<FAccountId>>(FOnlineAsyncOp& InAsyncOp, const TArray<EOS_EpicAccountId>& EpicAccountIds)> FAuthEOS::ResolveEpicIdsFn()
diff --git a/Engine/Plugins/Online/Experimental/OnlineServicesEOS/Source/Private/Online/CommerceEOS.cpp b/Engine/Plugins/Online/Experimental/OnlineServicesEOS/Source/Private/Online/CommerceEOS.cpp
--- a/Engine/Plugins/Online/Experimental/OnlineServicesEOS/Source/Private/Online/CommerceEOS.cpp
+++ b/Engine/Plugins/Online/Experimental/OnlineServicesEOS/Source/Private/Online/CommerceEOS.cpp
@@ -60,12 +60,12 @@ TOnlineAsyncOpHandle<FCommerceQueryOffers> FCommerceEOS::QueryOffers(FCommerceQu
 	Op->Then([this](TOnlineAsyncOp<FCommerceQueryOffers>& Op, TPromise<const EOS_Ecom_QueryOffersCallbackInfo*>&& Promise)
 	{
 		const FCommerceQueryOffers::Params& Params = Op.GetParams();
-		if (!Services.Get<FAuthEOS>()->IsLoggedIn(Params.LocalUserId))
+		if (!Services.Get<FAuthEOS>()->IsLoggedIn(Params.LocalAccountId))
 		{
 			Op.SetError(Errors::NotLoggedIn());
 			return;
 		}
-		EOS_EpicAccountId LocalUserEasId = GetEpicAccountId(Params.LocalUserId);
+		EOS_EpicAccountId LocalUserEasId = GetEpicAccountId(Params.LocalAccountId);
 		if (!EOS_EpicAccountId_IsValid(LocalUserEasId))
 		{
 			Op.SetError(Errors::NotLoggedIn());
@@ -91,7 +91,7 @@ TOnlineAsyncOpHandle<FCommerceQueryOffers> FCommerceEOS::QueryOffers(FCommerceQu
 		CountOptions.LocalUserId = Data->LocalUserId;
 		uint32 OfferCount = EOS_Ecom_GetOfferCount(EcomHandle, &CountOptions);
 
-		TArray<FOffer>& Offers = CachedOffers.FindOrAdd(Op.GetParams().LocalUserId, TArray<FOffer>());
+		TArray<FOffer>& Offers = CachedOffers.FindOrAdd(Op.GetParams().LocalAccountId, TArray<FOffer>());
 		Offers.Empty(OfferCount);
 
 		EOS_Ecom_CopyOfferByIndexOptions OfferOptions = { };
@@ -124,7 +124,7 @@ TOnlineAsyncOpHandle<FCommerceQueryOffersById> FCommerceEOS::QueryOffersById(FCo
 	TOnlineAsyncOpRef<FCommerceQueryOffersById> Op = GetOp<FCommerceQueryOffersById>(MoveTemp(Params));
 
 	// EOS doesn't support anything ID-specific, just use the generic QueryOffers call
-	QueryOffers({Op->GetParams().LocalUserId})
+	QueryOffers({Op->GetParams().LocalAccountId})
 	.OnComplete([this, Op](const TOnlineResult<FCommerceQueryOffers>& Result)
 	{
 		if(Result.IsOk())
@@ -143,18 +143,18 @@ TOnlineAsyncOpHandle<FCommerceQueryOffersById> FCommerceEOS::QueryOffersById(FCo
 
 TOnlineResult<FCommerceGetOffers> FCommerceEOS::GetOffers(FCommerceGetOffers::Params&& Params)
 {
-	if(CachedOffers.Contains(Params.LocalUserId))
+	if(CachedOffers.Contains(Params.LocalAccountId))
 	{
-		return TOnlineResult<FCommerceGetOffers>({CachedOffers.FindChecked(Params.LocalUserId)});
+		return TOnlineResult<FCommerceGetOffers>({CachedOffers.FindChecked(Params.LocalAccountId)});
 	}
 	return TOnlineResult<FCommerceGetOffers>(Errors::NotFound());
 }
 
 TOnlineResult<FCommerceGetOffersById> FCommerceEOS::GetOffersById(FCommerceGetOffersById::Params&& Params)
 {
-	if (CachedOffers.Contains(Params.LocalUserId))
+	if (CachedOffers.Contains(Params.LocalAccountId))
 	{
-		return TOnlineResult<FCommerceGetOffersById>({ CachedOffers.FindChecked(Params.LocalUserId).FilterByPredicate(
+		return TOnlineResult<FCommerceGetOffersById>({ CachedOffers.FindChecked(Params.LocalAccountId).FilterByPredicate(
 			[&Params](const FOffer& Offer)
 			{ 
 				return Params.OfferIds.Contains(Offer.OfferId);
@@ -179,12 +179,12 @@ TOnlineAsyncOpHandle<FCommerceCheckout> FCommerceEOS::Checkout(FCommerceCheckout
 	Op->Then([this](TOnlineAsyncOp<FCommerceCheckout>& Op, TPromise<const EOS_Ecom_CheckoutCallbackInfo*>&& Promise)
 	{
 		const FCommerceCheckout::Params& Params = Op.GetParams();
-		if (!Services.Get<FAuthEOS>()->IsLoggedIn(Params.LocalUserId))
+		if (!Services.Get<FAuthEOS>()->IsLoggedIn(Params.LocalAccountId))
 		{
 			Op.SetError(Errors::NotLoggedIn());
 			return;
 		}
-		EOS_EpicAccountId LocalUserEasId = GetEpicAccountId(Params.LocalUserId);
+		EOS_EpicAccountId LocalUserEasId = GetEpicAccountId(Params.LocalAccountId);
 		if (!EOS_EpicAccountId_IsValid(LocalUserEasId))
 		{
 			Op.SetError(Errors::NotLoggedIn());
@@ -236,12 +236,12 @@ TOnlineAsyncOpHandle<FCommerceQueryTransactionEntitlements> FCommerceEOS::QueryT
 	Op->Then([this](TOnlineAsyncOp<FCommerceQueryTransactionEntitlements>& Op)
 	{
 		const FCommerceQueryTransactionEntitlements::Params& Params = Op.GetParams();
-		if (!Services.Get<FAuthEOS>()->IsLoggedIn(Params.LocalUserId))
+		if (!Services.Get<FAuthEOS>()->IsLoggedIn(Params.LocalAccountId))
 		{
 			Op.SetError(Errors::NotLoggedIn());
 			return;
 		}
-		EOS_EpicAccountId LocalUserEasId = GetEpicAccountId(Params.LocalUserId);
+		EOS_EpicAccountId LocalUserEasId = GetEpicAccountId(Params.LocalAccountId);
 		if (!EOS_EpicAccountId_IsValid(LocalUserEasId))
 		{
 			Op.SetError(Errors::NotLoggedIn());
@@ -309,12 +309,12 @@ TOnlineAsyncOpHandle<FCommerceQueryEntitlements> FCommerceEOS::QueryEntitlements
 	Op->Then([this](TOnlineAsyncOp<FCommerceQueryEntitlements>& Op, TPromise<const EOS_Ecom_QueryEntitlementsCallbackInfo*>&& Promise)
 	{
 		const FCommerceQueryEntitlements::Params& Params = Op.GetParams();
-		if (!Services.Get<FAuthEOS>()->IsLoggedIn(Params.LocalUserId))
+		if (!Services.Get<FAuthEOS>()->IsLoggedIn(Params.LocalAccountId))
 		{
 			Op.SetError(Errors::NotLoggedIn());
 			return;
 		}
-		EOS_EpicAccountId LocalUserEasId = GetEpicAccountId(Params.LocalUserId);
+		EOS_EpicAccountId LocalUserEasId = GetEpicAccountId(Params.LocalAccountId);
 		if (!EOS_EpicAccountId_IsValid(LocalUserEasId))
 		{
 			Op.SetError(Errors::NotLoggedIn());
@@ -342,7 +342,7 @@ TOnlineAsyncOpHandle<FCommerceQueryEntitlements> FCommerceEOS::QueryEntitlements
 		CountOptions.ApiVersion = EOS_ECOM_GETENTITLEMENTSCOUNT_API_LATEST;
 		CountOptions.LocalUserId = Data->LocalUserId;
 		uint32 Count = EOS_Ecom_GetEntitlementsCount(EcomHandle, &CountOptions);
-		TArray<FEntitlement>& Entitlements = CachedEntitlements.FindOrAdd(Op.GetParams().LocalUserId, TArray<FEntitlement>());
+		TArray<FEntitlement>& Entitlements = CachedEntitlements.FindOrAdd(Op.GetParams().LocalAccountId, TArray<FEntitlement>());
 		Entitlements.Reset(Count);
 
 		EOS_Ecom_CopyEntitlementByIndexOptions CopyOptions = { };
@@ -377,11 +377,11 @@ TOnlineAsyncOpHandle<FCommerceQueryEntitlements> FCommerceEOS::QueryEntitlements
 
 TOnlineResult<FCommerceGetEntitlements> FCommerceEOS::GetEntitlements(FCommerceGetEntitlements::Params&& Params)
 {
-	if (!CachedEntitlements.Contains(Params.LocalUserId))
+	if (!CachedEntitlements.Contains(Params.LocalAccountId))
 	{
 		return TOnlineResult<FCommerceGetEntitlements>(Errors::NotFound());
 	}
-	return TOnlineResult<FCommerceGetEntitlements>({CachedEntitlements.FindChecked(Params.LocalUserId)});
+	return TOnlineResult<FCommerceGetEntitlements>({CachedEntitlements.FindChecked(Params.LocalAccountId)});
 }
 
 TOnlineAsyncOpHandle<FCommerceRedeemEntitlement> FCommerceEOS::RedeemEntitlement(FCommerceRedeemEntitlement::Params&& Params)
@@ -391,12 +391,12 @@ TOnlineAsyncOpHandle<FCommerceRedeemEntitlement> FCommerceEOS::RedeemEntitlement
 	Op->Then([this](TOnlineAsyncOp<FCommerceRedeemEntitlement>& Op, TPromise<const EOS_Ecom_RedeemEntitlementsCallbackInfo*>&& Promise)
 	{
 		const FCommerceRedeemEntitlement::Params& Params = Op.GetParams();
-		if (!Services.Get<FAuthEOS>()->IsLoggedIn(Params.LocalUserId))
+		if (!Services.Get<FAuthEOS>()->IsLoggedIn(Params.LocalAccountId))
 		{
 			Op.SetError(Errors::NotLoggedIn());
 			return;
 		}
-		EOS_EpicAccountId LocalUserEasId = GetEpicAccountId(Params.LocalUserId);
+		EOS_EpicAccountId LocalUserEasId = GetEpicAccountId(Params.LocalAccountId);
 		if (!EOS_EpicAccountId_IsValid(LocalUserEasId))
 		{
 			Op.SetError(Errors::NotLoggedIn());
@@ -436,12 +436,12 @@ TOnlineAsyncOpHandle<FCommerceRetrieveS2SToken> FCommerceEOS::RetrieveS2SToken(F
 	Op->Then([this](TOnlineAsyncOp<FCommerceRetrieveS2SToken>& Op, TPromise<const EOS_Ecom_QueryOwnershipTokenCallbackInfo*>&& Promise)
 	{
 		const FCommerceRetrieveS2SToken::Params& Params = Op.GetParams();
-		if (!Services.Get<FAuthEOS>()->IsLoggedIn(Params.LocalUserId))
+		if (!Services.Get<FAuthEOS>()->IsLoggedIn(Params.LocalAccountId))
 		{
 			Op.SetError(Errors::NotLoggedIn());
 			return;
 		}
-		EOS_EpicAccountId LocalUserEasId = GetEpicAccountId(Params.LocalUserId);
+		EOS_EpicAccountId LocalUserEasId = GetEpicAccountId(Params.LocalAccountId);
 		if (!EOS_EpicAccountId_IsValid(LocalUserEasId))
 		{
 			Op.SetError(Errors::NotLoggedIn());
diff --git a/Engine/Plugins/Online/Experimental/OnlineServicesEOS/Source/Private/Online/ExternalUIEOS.cpp b/Engine/Plugins/Online/Experimental/OnlineServicesEOS/Source/Private/Online/ExternalUIEOS.cpp
--- a/Engine/Plugins/Online/Experimental/OnlineServicesEOS/Source/Private/Online/ExternalUIEOS.cpp
+++ b/Engine/Plugins/Online/Experimental/OnlineServicesEOS/Source/Private/Online/ExternalUIEOS.cpp
@@ -84,10 +84,10 @@ TOnlineAsyncOpHandle<FExternalUIShowFriendsUI> FExternalUIEOS::ShowFriendsUI(FEx
 {
 	TOnlineAsyncOpRef<FExternalUIShowFriendsUI> Op = GetOp<FExternalUIShowFriendsUI>(MoveTemp(Params));
 
-	EOS_EpicAccountId LocalUserEasId = GetEpicAccountId(Params.LocalUserId);
+	EOS_EpicAccountId LocalUserEasId = GetEpicAccountId(Params.LocalAccountId);
 	if (!EOS_EpicAccountId_IsValid(LocalUserEasId))
 	{
-		UE_LOG(LogTemp, Warning, TEXT("[FExternalUIEOS::ShowFriendsUI] LocalUserId=[%s] EpicAccountId not found"), *ToLogString(Params.LocalUserId));
+		UE_LOG(LogTemp, Warning, TEXT("[FExternalUIEOS::ShowFriendsUI] LocalAccountId=[%s] EpicAccountId not found"), *ToLogString(Params.LocalAccountId));
 		Op->SetError(Errors::Unknown()); // TODO
 		return Op->GetHandle();
 	}
diff --git a/Engine/Plugins/Online/Experimental/OnlineServicesEOS/Source/Private/Online/FriendsEOS.cpp b/Engine/Plugins/Online/Experimental/OnlineServicesEOS/Source/Private/Online/FriendsEOS.cpp
--- a/Engine/Plugins/Online/Experimental/OnlineServicesEOS/Source/Private/Online/FriendsEOS.cpp
+++ b/Engine/Plugins/Online/Experimental/OnlineServicesEOS/Source/Private/Online/FriendsEOS.cpp
@@ -45,12 +45,12 @@ void FFriendsEOS::Initialize()
 	{
 		FFriendsEOS* This = reinterpret_cast<FFriendsEOS*>(Data->ClientData);
 
-		const FAccountId LocalUserId = FindAccountIdChecked(Data->LocalUserId);
-		This->Services.Get<FAuthEOS>()->ResolveAccountId(LocalUserId, Data->TargetUserId)
-		.Next([This, LocalUserId, PreviousStatus = Data->PreviousStatus, CurrentStatus = Data->CurrentStatus](const FAccountId& FriendUserId)
+		const FAccountId LocalAccountId = FindAccountIdChecked(Data->LocalUserId);
+		This->Services.Get<FAuthEOS>()->ResolveAccountId(LocalAccountId, Data->TargetUserId)
+		.Next([This, LocalAccountId, PreviousStatus = Data->PreviousStatus, CurrentStatus = Data->CurrentStatus](const FAccountId& FriendAccountId)
 		{
-			UE_LOG(LogTemp, Warning, TEXT("OnEOSFriendsUpdate: LocalUserId=[%s] FriendUserId=[%s]"), *ToLogString(LocalUserId), *ToLogString(FriendUserId));
-			This->OnEOSFriendsUpdate(LocalUserId, FriendUserId, PreviousStatus, CurrentStatus);
+			UE_LOG(LogTemp, Warning, TEXT("OnEOSFriendsUpdate: LocalAccountId=[%s] FriendAccountId=[%s]"), *ToLogString(LocalAccountId), *ToLogString(FriendAccountId));
+			This->OnEOSFriendsUpdate(LocalAccountId, FriendAccountId, PreviousStatus, CurrentStatus);
 		});
 	});
 }
@@ -67,7 +67,7 @@ TOnlineAsyncOpHandle<FQueryFriends> FFriendsEOS::QueryFriends(FQueryFriends::Par
 	{
 		// Initialize
 		const FQueryFriends::Params& Params = Op->GetParams();
-		if (!Services.Get<FAuthEOS>()->IsLoggedIn(Params.LocalUserId))
+		if (!Services.Get<FAuthEOS>()->IsLoggedIn(Params.LocalAccountId))
 		{
 			// TODO: Error codes
 			Op->SetError(Errors::Unknown());
@@ -76,7 +76,7 @@ TOnlineAsyncOpHandle<FQueryFriends> FFriendsEOS::QueryFriends(FQueryFriends::Par
 
 		EOS_Friends_QueryFriendsOptions QueryFriendsOptions = { };
 		QueryFriendsOptions.ApiVersion = EOS_FRIENDS_QUERYFRIENDS_API_LATEST;
-		EOS_EpicAccountId LocalUserEasId = GetEpicAccountIdChecked(Params.LocalUserId);
+		EOS_EpicAccountId LocalUserEasId = GetEpicAccountIdChecked(Params.LocalAccountId);
 		QueryFriendsOptions.LocalUserId = LocalUserEasId;
 
 		Op->Then([this, QueryFriendsOptions](TOnlineAsyncOp<FQueryFriends>& InAsyncOp, TPromise<const EOS_Friends_QueryFriendsCallbackInfo*>&& Promise) mutable
@@ -117,10 +117,10 @@ TOnlineAsyncOpHandle<FQueryFriends> FFriendsEOS::QueryFriends(FQueryFriends::Par
 		}).Then(Services.Get<FAuthEOS>()->ResolveEpicIdsFn())
 		.Then([this](TOnlineAsyncOp<FQueryFriends>& InAsyncOp, const TArray<FAccountId>& CurrentFriendIds) mutable
 		{
-			const FAccountId LocalUserId = InAsyncOp.GetParams().LocalUserId;
-			const EOS_EpicAccountId LocalUserEasId = GetEpicAccountIdChecked(LocalUserId);
+			const FAccountId LocalAccountId = InAsyncOp.GetParams().LocalAccountId;
+			const EOS_EpicAccountId LocalUserEasId = GetEpicAccountIdChecked(LocalAccountId);
 
-			TMap<FAccountId, TSharedRef<FFriend>>& FriendsList = FriendsLists.FindOrAdd(LocalUserId);
+			TMap<FAccountId, TSharedRef<FFriend>>& FriendsList = FriendsLists.FindOrAdd(LocalAccountId);
 			TArray<FAccountId> NewFriendIds;
 			TArray<FAccountId> PreviousFriendIds;
 			TArray<FAccountId> UpdatedFriendIds; // TODO:  More granular updates like what changed?
@@ -154,7 +154,7 @@ TOnlineAsyncOpHandle<FQueryFriends> FFriendsEOS::QueryFriends(FQueryFriends::Par
 					// Add new friend
 					UE_LOG(LogTemp, Warning, TEXT("QueryFriendsComplete: Adding %s"), *LexToString(FriendEasId));
 					TSharedRef<FFriend> Friend = MakeShared<FFriend>();
-					Friend->UserId = FriendId;
+					Friend->AccountId = FriendId;
 					Friend->InviteStatus = InviteStatus;
 					FriendsList.Emplace(FriendId, MoveTemp(Friend));
 					NewFriendIds.Emplace(FriendId);
@@ -177,7 +177,7 @@ TOnlineAsyncOpHandle<FQueryFriends> FFriendsEOS::QueryFriends(FQueryFriends::Par
 			if (bAnyChanges)
 			{
 				FFriendsListUpdated FriendsListUpdatedParams = {};
-				FriendsListUpdatedParams.LocalUserId = InAsyncOp.GetParams().LocalUserId;
+				FriendsListUpdatedParams.LocalAccountId = InAsyncOp.GetParams().LocalAccountId;
 				FriendsListUpdatedParams.AddedFriends = NewFriendIds;
 				FriendsListUpdatedParams.RemovedFriends = PreviousFriendIds;
 				FriendsListUpdatedParams.UpdatedFriends = UpdatedFriendIds;
@@ -191,7 +191,7 @@ TOnlineAsyncOpHandle<FQueryFriends> FFriendsEOS::QueryFriends(FQueryFriends::Par
 
 TOnlineResult<FGetFriends> FFriendsEOS::GetFriends(FGetFriends::Params&& Params)
 {
-	if (TMap<FAccountId, TSharedRef<FFriend>>* FriendsList = FriendsLists.Find(Params.LocalUserId))
+	if (TMap<FAccountId, TSharedRef<FFriend>>* FriendsList = FriendsLists.Find(Params.LocalAccountId))
 	{
 		FGetFriends::Result Result;
 		FriendsList->GenerateValueArray(Result.Friends);
@@ -208,7 +208,7 @@ TOnlineAsyncOpHandle<FAddFriend> FFriendsEOS::AddFriend(FAddFriend::Params&& InP
 #if 0 // At present this causes issues due to the EOS function executing the callback immediately, before the future can be bound to
 		// Initialize
 		const FAddFriend::Params& Params = Op->GetParams();
-		if (!Services.Get<FAuthEOS>()->IsLoggedIn(Params.LocalUserId))
+		if (!Services.Get<FAuthEOS>()->IsLoggedIn(Params.LocalAccountId))
 		{
 			// TODO: Error codes
 			Op->SetError(Errors::UnknownError());
@@ -229,7 +229,7 @@ TOnlineAsyncOpHandle<FAddFriend> FFriendsEOS::AddFriend(FAddFriend::Params&& InP
 			const FAddFriend::Params& Params = InAsyncOp.GetParams();
 			EOS_Friends_SendInviteOptions SendInviteOptions = { };
 			SendInviteOptions.ApiVersion = EOS_FRIENDS_SENDINVITE_API_LATEST;
-			SendInviteOptions.LocalUserId = EOSAccountIdFromOnlineServiceAccountId(Params.LocalUserId).GetValue();
+			SendInviteOptions.LocalUserId = EOSAccountIdFromOnlineServiceAccountId(Params.LocalAccountId).GetValue();
 			SendInviteOptions.TargetUserId = EOSAccountIdFromOnlineServiceAccountId(Params.FriendId).GetValue();
 			EOS_Async(EOS_Friends_SendInvite, FriendsHandle, SendInviteOptions, MoveTemp(Promise));
 		})
@@ -245,33 +245,33 @@ TOnlineAsyncOpHandle<FAddFriend> FFriendsEOS::AddFriend(FAddFriend::Params&& InP
 	return Op->GetHandle();
 }
 
-void FFriendsEOS::OnEOSFriendsUpdate(FAccountId LocalUserId, FAccountId FriendUserId, EOS_EFriendsStatus PreviousStatus, EOS_EFriendsStatus CurrentStatus)
+void FFriendsEOS::OnEOSFriendsUpdate(FAccountId LocalAccountId, FAccountId FriendAccountId, EOS_EFriendsStatus PreviousStatus, EOS_EFriendsStatus CurrentStatus)
 {
 	// TODO
-	UE_LOG(LogTemp, Warning, TEXT("OnEOSFriendsUpdate: LocalUserId=[%s] FriendUserId=[%s] PreviousStatus=[%s] CurrentStatus=[%s]"), *ToLogString(LocalUserId), *ToLogString(FriendUserId), LexToString(PreviousStatus), LexToString(CurrentStatus));
+	UE_LOG(LogTemp, Warning, TEXT("OnEOSFriendsUpdate: LocalAccountId=[%s] FriendAccountId=[%s] PreviousStatus=[%s] CurrentStatus=[%s]"), *ToLogString(LocalAccountId), *ToLogString(FriendAccountId), LexToString(PreviousStatus), LexToString(CurrentStatus));
 
-	TMap<FAccountId, TSharedRef<FFriend>>& FriendsList = FriendsLists.FindOrAdd(LocalUserId);
+	TMap<FAccountId, TSharedRef<FFriend>>& FriendsList = FriendsLists.FindOrAdd(LocalAccountId);
 	bool bAnyChanges = false;
 	TArray<FAccountId, TInlineAllocator<1>> AddedFriends;
 	TArray<FAccountId, TInlineAllocator<1>> RemovedFriends;
 	TArray<FAccountId, TInlineAllocator<1>> UpdatedFriends;
 	if (CurrentStatus == EOS_EFriendsStatus::EOS_FS_NotFriends)
 	{
-		bAnyChanges = FriendsList.Remove(FriendUserId) > 0;
+		bAnyChanges = FriendsList.Remove(FriendAccountId) > 0;
 		if (bAnyChanges)
 		{
-			RemovedFriends.Emplace(FriendUserId);
+			RemovedFriends.Emplace(FriendAccountId);
 		}
 	}
 	else
 	{
 		TOptional<EFriendInviteStatus> InviteStatus = EOSFriendStatusToInviteStatus(CurrentStatus);
-		if (TSharedRef<FFriend>* ExistingFriend = FriendsList.Find(FriendUserId))
+		if (TSharedRef<FFriend>* ExistingFriend = FriendsList.Find(FriendAccountId))
 		{
 			if ((*ExistingFriend)->InviteStatus != InviteStatus)
 			{
 				(*ExistingFriend)->InviteStatus = InviteStatus;
-				UpdatedFriends.Emplace(FriendUserId);
+				UpdatedFriends.Emplace(FriendAccountId);
 				bAnyChanges = true;
 			}
 		}
@@ -279,11 +279,11 @@ void FFriendsEOS::OnEOSFriendsUpdate(FAccountId LocalUserId, FAccountId FriendUs
 		{
 			// Create new entry
 			TSharedRef<FFriend> Friend = MakeShared<FFriend>();
-			Friend->UserId = FriendUserId;
+			Friend->AccountId = FriendAccountId;
 			Friend->InviteStatus = InviteStatus;
-			FriendsList.Emplace(FriendUserId, MoveTemp(Friend));
+			FriendsList.Emplace(FriendAccountId, MoveTemp(Friend));
 
-			AddedFriends.Emplace(FriendUserId);
+			AddedFriends.Emplace(FriendAccountId);
 			bAnyChanges = true;
 		}
 	}
@@ -291,7 +291,7 @@ void FFriendsEOS::OnEOSFriendsUpdate(FAccountId LocalUserId, FAccountId FriendUs
 	if (bAnyChanges)
 	{
 		FFriendsListUpdated FriendsListUpdatedParams = {};
-		FriendsListUpdatedParams.LocalUserId = LocalUserId;
+		FriendsListUpdatedParams.LocalAccountId = LocalAccountId;
 		FriendsListUpdatedParams.AddedFriends = AddedFriends;
 		FriendsListUpdatedParams.RemovedFriends = RemovedFriends;
 		FriendsListUpdatedParams.UpdatedFriends = UpdatedFriends;
diff --git a/Engine/Plugins/Online/Experimental/OnlineServicesEOS/Source/Private/Online/OnlineIdEOS.cpp b/Engine/Plugins/Online/Experimental/OnlineServicesEOS/Source/Private/Online/OnlineIdEOS.cpp
--- a/Engine/Plugins/Online/Experimental/OnlineServicesEOS/Source/Private/Online/OnlineIdEOS.cpp
+++ b/Engine/Plugins/Online/Experimental/OnlineServicesEOS/Source/Private/Online/OnlineIdEOS.cpp
@@ -9,17 +9,17 @@
 
 namespace UE::Online {
 
-FString FOnlineAccountIdRegistryEOS::ToLogString(const FAccountId& Handle) const
+FString FOnlineAccountIdRegistryEOS::ToLogString(const FAccountId& AccountId) const
 {
 	FString Result;
-	if (ValidateOnlineId(Handle))
+	if (ValidateOnlineId(AccountId))
 	{
-		const FOnlineAccountIdDataEOS& IdData = GetAccountIdData(Handle);
+		const FOnlineAccountIdDataEOS& IdData = GetAccountIdData(AccountId);
 		Result = FString::Printf(TEXT("EAS=[%s] EOS=[%s]"), *LexToString(IdData.EpicAccountId), *LexToString(IdData.ProductUserId));
 	}
 	else
 	{
-		check(!Handle.IsValid()); // Check we haven't been passed a valid handle for a different EOnlineServices.
+		check(!AccountId.IsValid()); // Check we haven't been passed a valid handle for a different EOnlineServices.
 		Result = TEXT("Invalid");
 	}
 	return Result;
@@ -35,12 +35,12 @@ ENUM_CLASS_FLAGS(EEOSAccountIdElements);
 const uint8 OnlineIdEOSUtf8BufferLength = 32;
 const uint8 OnlineIdEOSHexBufferLength = 16;
 
-TArray<uint8> FOnlineAccountIdRegistryEOS::ToReplicationData(const FAccountId& Handle) const
+TArray<uint8> FOnlineAccountIdRegistryEOS::ToReplicationData(const FAccountId& AccountId) const
 {
 	TArray<uint8> ReplicationData;
-	if (ValidateOnlineId(Handle))
+	if (ValidateOnlineId(AccountId))
 	{
-		const FOnlineAccountIdDataEOS& IdData = GetAccountIdData(Handle);
+		const FOnlineAccountIdDataEOS& IdData = GetAccountIdData(AccountId);
 		EEOSAccountIdElements Elements = EEOSAccountIdElements::None;
 
 		char EasBuffer[EOS_EPICACCOUNTID_MAX_LENGTH + 1] = { 0 };
@@ -82,7 +82,7 @@ TArray<uint8> FOnlineAccountIdRegistryEOS::ToReplicationData(const FAccountId& H
 	}
 	else
 	{
-		check(!Handle.IsValid()); // Check we haven't been passed a valid handle for a different EOnlineServices.
+		check(!AccountId.IsValid()); // Check we haven't been passed a valid handle for a different EOnlineServices.
 	}
 	return ReplicationData;
 }
@@ -114,43 +114,43 @@ FAccountId FOnlineAccountIdRegistryEOS::FromReplicationData(const TArray<uint8>&
 
 FAccountId FOnlineAccountIdRegistryEOS::FindAccountId(const EOS_ProductUserId ProductUserId) const
 {
-	FAccountId Handle;
+	FAccountId AccountId;
 	if (EOS_ProductUserId_IsValid(ProductUserId) == EOS_TRUE)
 	{
 		const FReadScopeLock ReadLock(Lock);
-		if (const FAccountId* Found = ProductUserIdToHandle.Find(ProductUserId))
+		if (const FAccountId* Found = ProductUserIdToAccountId.Find(ProductUserId))
 		{
-			Handle = *Found;
+			AccountId = *Found;
 		}
 	}
-	return Handle;
+	return AccountId;
 }
 
-EOS_ProductUserId FOnlineAccountIdRegistryEOS::GetProductUserId(const FAccountId& Handle) const
+EOS_ProductUserId FOnlineAccountIdRegistryEOS::GetProductUserId(const FAccountId& AccountId) const
 {
-	return GetAccountIdData(Handle).ProductUserId;
+	return GetAccountIdData(AccountId).ProductUserId;
 }
 
 FAccountId FOnlineAccountIdRegistryEOS::FindAccountId(const EOS_EpicAccountId EpicAccountId) const
 {
-	FAccountId Handle;
+	FAccountId AccountId;
 	if (EOS_EpicAccountId_IsValid(EpicAccountId) == EOS_TRUE)
 	{
 		const FReadScopeLock ReadLock(Lock);
-		if (const FAccountId* Found = EpicAccountIdToHandle.Find(EpicAccountId))
+		if (const FAccountId* Found = EpicAccountIdToAccountId.Find(EpicAccountId))
 		{
-			Handle = *Found;
+			AccountId = *Found;
 		}
 	}
-	return Handle;
+	return AccountId;
 }
 
-FOnlineAccountIdDataEOS FOnlineAccountIdRegistryEOS::GetAccountIdData(const FAccountId& Handle) const
+FOnlineAccountIdDataEOS FOnlineAccountIdRegistryEOS::GetAccountIdData(const FAccountId& AccountId) const
 {
-	if (ValidateOnlineId(Handle))
+	if (ValidateOnlineId(AccountId))
 	{
 		FReadScopeLock ScopeLock(Lock);
-		return AccountIdData[Handle.GetHandle() - 1];
+		return AccountIdData[AccountId.GetHandle() - 1];
 	}
 	return FOnlineAccountIdDataEOS();
 }
@@ -174,11 +174,11 @@ FAccountId FOnlineAccountIdRegistryEOS::FindOrAddAccountId(const EOS_EpicAccount
 		auto FindExisting = [this, InEpicAccountId, InProductUserId, bInEpicAccountIdValid, bInProductUserIdValid, &bUpdateEpicAccountId, &bUpdateProductUserId]()
 		{
 			FAccountId Result;
-			if (const FAccountId* FoundEas = bInEpicAccountIdValid ? EpicAccountIdToHandle.Find(InEpicAccountId) : nullptr)
+			if (const FAccountId* FoundEas = bInEpicAccountIdValid ? EpicAccountIdToAccountId.Find(InEpicAccountId) : nullptr)
 			{
 				Result = *FoundEas;
 			}
-			else if (const FAccountId* FoundProd = bInProductUserIdValid ? ProductUserIdToHandle.Find(InProductUserId) : nullptr)
+			else if (const FAccountId* FoundProd = bInProductUserIdValid ? ProductUserIdToAccountId.Find(InProductUserId) : nullptr)
 			{
 				Result = *FoundProd;
 			}
@@ -222,11 +222,11 @@ FAccountId FOnlineAccountIdRegistryEOS::FindOrAddAccountId(const EOS_EpicAccount
 
 				if (bInEpicAccountIdValid)
 				{
-					EpicAccountIdToHandle.Emplace(InEpicAccountId, Result);
+					EpicAccountIdToAccountId.Emplace(InEpicAccountId, Result);
 				}
 				if (bInProductUserIdValid)
 				{
-					ProductUserIdToHandle.Emplace(InProductUserId, Result);
+					ProductUserIdToAccountId.Emplace(InProductUserId, Result);
 				}
 			}
 		}
@@ -240,12 +240,12 @@ FAccountId FOnlineAccountIdRegistryEOS::FindOrAddAccountId(const EOS_EpicAccount
 			if (bUpdateEpicAccountId)
 			{
 				AccountIdDataToUpdate.EpicAccountId = InEpicAccountId;
-				EpicAccountIdToHandle.Emplace(InEpicAccountId, Result);
+				EpicAccountIdToAccountId.Emplace(InEpicAccountId, Result);
 			}
 			if (bUpdateProductUserId)
 			{
 				AccountIdDataToUpdate.ProductUserId = InProductUserId;
-				ProductUserIdToHandle.Emplace(InProductUserId, Result);
+				ProductUserIdToAccountId.Emplace(InProductUserId, Result);
 			}
 		}
 	}
@@ -253,14 +253,14 @@ FAccountId FOnlineAccountIdRegistryEOS::FindOrAddAccountId(const EOS_EpicAccount
 	return Result;
 }
 
-EOS_EpicAccountId GetEpicAccountId(const FAccountId& Handle)
+EOS_EpicAccountId GetEpicAccountId(const FAccountId& AccountId)
 {
-	return FOnlineAccountIdRegistryEOS::Get().GetAccountIdData(Handle).EpicAccountId;
+	return FOnlineAccountIdRegistryEOS::Get().GetAccountIdData(AccountId).EpicAccountId;
 }
 
-EOS_EpicAccountId GetEpicAccountIdChecked(const FAccountId& Handle)
+EOS_EpicAccountId GetEpicAccountIdChecked(const FAccountId& AccountId)
 {
-	EOS_EpicAccountId EpicAccountId = GetEpicAccountId(Handle);
+	EOS_EpicAccountId EpicAccountId = GetEpicAccountId(AccountId);
 	check(EOS_EpicAccountId_IsValid(EpicAccountId) == EOS_TRUE);
 	return EpicAccountId;
 }
diff --git a/Engine/Plugins/Online/Experimental/OnlineServicesEOS/Source/Private/Online/OnlineIdEOS.h b/Engine/Plugins/Online/Experimental/OnlineServicesEOS/Source/Private/Online/OnlineIdEOS.h
--- a/Engine/Plugins/Online/Experimental/OnlineServicesEOS/Source/Private/Online/OnlineIdEOS.h
+++ b/Engine/Plugins/Online/Experimental/OnlineServicesEOS/Source/Private/Online/OnlineIdEOS.h
@@ -24,20 +24,20 @@ public:
 	virtual ~FOnlineAccountIdRegistryEOS() = default;
 
 	// Begin IOnlineAccountIdRegistry
-	virtual FString ToLogString(const FAccountId& Handle) const override;
-	virtual TArray<uint8> ToReplicationData(const FAccountId& Handle) const override;
+	virtual FString ToLogString(const FAccountId& AccountId) const override;
+	virtual TArray<uint8> ToReplicationData(const FAccountId& AccountId) const override;
 	virtual FAccountId FromReplicationData(const TArray<uint8>& ReplicationString) override;
 	// End IOnlineAccountIdRegistry
 
 	// Begin IOnlineAccountRegistryEOSGS
 	virtual FAccountId FindAccountId(const EOS_ProductUserId ProductUserId) const override;
-	virtual EOS_ProductUserId GetProductUserId(const FAccountId& Handle) const override;
+	virtual EOS_ProductUserId GetProductUserId(const FAccountId& AccountId) const override;
 	// End IOnlineAccountRegistryEOSGS
 
 	FAccountId FindAccountId(const EOS_EpicAccountId EpicAccountId) const;
 
 	// Return copies as it is not thread safe to return pointers/references to array elements, in case the array is grown+relocated on another thread.
-	FOnlineAccountIdDataEOS GetAccountIdData(const FAccountId& Handle) const;
+	FOnlineAccountIdDataEOS GetAccountIdData(const FAccountId& AccountId) const;
 
 	static FOnlineAccountIdRegistryEOS& Get();
 
@@ -51,12 +51,12 @@ private:
 
 	TArray<FOnlineAccountIdDataEOS> AccountIdData; // Actual container for the info, indexed by the handle
 
-	TMap<EOS_EpicAccountId, FAccountId> EpicAccountIdToHandle; // Map of EOS_EpicAccountId to the associated handle.
-	TMap<EOS_ProductUserId, FAccountId> ProductUserIdToHandle; // Map of EOS_ProductUserId to the associated handle.
+	TMap<EOS_EpicAccountId, FAccountId> EpicAccountIdToAccountId; // Map of EOS_EpicAccountId to the associated handle.
+	TMap<EOS_ProductUserId, FAccountId> ProductUserIdToAccountId; // Map of EOS_ProductUserId to the associated handle.
 };
 
-EOS_EpicAccountId GetEpicAccountId(const FAccountId& Handle);
-EOS_EpicAccountId GetEpicAccountIdChecked(const FAccountId& Handle);
+EOS_EpicAccountId GetEpicAccountId(const FAccountId& AccountId);
+EOS_EpicAccountId GetEpicAccountIdChecked(const FAccountId& AccountId);
 FAccountId FindAccountId(const EOS_EpicAccountId EpicAccountId);
 FAccountId FindAccountIdChecked(const EOS_EpicAccountId EpicAccountId);
 
diff --git a/Engine/Plugins/Online/Experimental/OnlineServicesEOS/Source/Private/Online/PresenceEOS.cpp b/Engine/Plugins/Online/Experimental/OnlineServicesEOS/Source/Private/Online/PresenceEOS.cpp
--- a/Engine/Plugins/Online/Experimental/OnlineServicesEOS/Source/Private/Online/PresenceEOS.cpp
+++ b/Engine/Plugins/Online/Experimental/OnlineServicesEOS/Source/Private/Online/PresenceEOS.cpp
@@ -78,13 +78,13 @@ void FPresenceEOS::Initialize()
 	NotifyPresenceChangedNotificationId = EOS_Presence_AddNotifyOnPresenceChanged(PresenceHandle, &Options, this, [](const EOS_Presence_PresenceChangedCallbackInfo* Data)
 	{
 		FPresenceEOS* This = reinterpret_cast<FPresenceEOS*>(Data->ClientData);
-		const FAccountId LocalUserId = FindAccountIdChecked(Data->LocalUserId);
+		const FAccountId LocalAccountId = FindAccountIdChecked(Data->LocalUserId);
 
-		This->Services.Get<FAuthEOS>()->ResolveAccountId(LocalUserId, Data->PresenceUserId)
-		.Next([This, LocalUserId](const FAccountId& PresenceUserId)
+		This->Services.Get<FAuthEOS>()->ResolveAccountId(LocalAccountId, Data->PresenceUserId)
+		.Next([This, LocalAccountId](const FAccountId& PresenceAccountId)
 		{
-			UE_LOG(LogTemp, Warning, TEXT("OnEOSPresenceUpdate: LocalUserId=[%s] PresenceUserId=[%s]"), *ToLogString(LocalUserId), *ToLogString(PresenceUserId));
-			This->UpdateUserPresence(LocalUserId, PresenceUserId);
+			UE_LOG(LogTemp, Warning, TEXT("OnEOSPresenceUpdate: LocalAccountId=[%s] PresenceAccountId=[%s]"), *ToLogString(LocalAccountId), *ToLogString(PresenceAccountId));
+			This->UpdateUserPresence(LocalAccountId, PresenceAccountId);
 		});
 	});
 }
@@ -101,12 +101,12 @@ TOnlineAsyncOpHandle<FQueryPresence> FPresenceEOS::QueryPresence(FQueryPresence:
 	{
 		// Initialize
 		const FQueryPresence::Params& Params = Op->GetParams();
-		if (!Services.Get<FAuthEOS>()->IsLoggedIn(Params.LocalUserId))
+		if (!Services.Get<FAuthEOS>()->IsLoggedIn(Params.LocalAccountId))
 		{
 			Op->SetError(Errors::NotLoggedIn());
 			return Op->GetHandle();
 		}
-		EOS_EpicAccountId TargetUserEasId = GetEpicAccountId(Params.TargetUserId);
+		EOS_EpicAccountId TargetUserEasId = GetEpicAccountId(Params.TargetAccountId);
 		if (!EOS_EpicAccountId_IsValid(TargetUserEasId))
 		{
 			Op->SetError(Errors::NotLoggedIn());
@@ -114,7 +114,7 @@ TOnlineAsyncOpHandle<FQueryPresence> FPresenceEOS::QueryPresence(FQueryPresence:
 		}
 
 		// TODO:  If we try to query a local user's presence, is that an error, should we return the cached state, should we still ask EOS?
-		const bool bIsLocalUser = Services.Get<FAuthEOS>()->IsLoggedIn(Params.TargetUserId);
+		const bool bIsLocalUser = Services.Get<FAuthEOS>()->IsLoggedIn(Params.TargetAccountId);
 		if (bIsLocalUser)
 		{
 			Op->SetError(Errors::CannotQueryLocalUsers()); 
@@ -126,8 +126,8 @@ TOnlineAsyncOpHandle<FQueryPresence> FPresenceEOS::QueryPresence(FQueryPresence:
 				const FQueryPresence::Params& Params = InAsyncOp.GetParams();
 				EOS_Presence_QueryPresenceOptions QueryPresenceOptions = { };
 				QueryPresenceOptions.ApiVersion = EOS_PRESENCE_QUERYPRESENCE_API_LATEST;
-				QueryPresenceOptions.LocalUserId = GetEpicAccountIdChecked(Params.LocalUserId);
-				QueryPresenceOptions.TargetUserId = GetEpicAccountIdChecked(Params.TargetUserId);
+				QueryPresenceOptions.LocalUserId = GetEpicAccountIdChecked(Params.LocalAccountId);
+				QueryPresenceOptions.TargetUserId = GetEpicAccountIdChecked(Params.TargetAccountId);
 
 				EOS_Async(EOS_Presence_QueryPresence, PresenceHandle, QueryPresenceOptions, MoveTemp(Promise));
 			})
@@ -138,8 +138,8 @@ TOnlineAsyncOpHandle<FQueryPresence> FPresenceEOS::QueryPresence(FQueryPresence:
 				if (Data->ResultCode == EOS_EResult::EOS_Success)
 				{
 					const FQueryPresence::Params& Params = InAsyncOp.GetParams();
-					UpdateUserPresence(Params.LocalUserId, Params.TargetUserId);
-					FQueryPresence::Result Result = { FindOrCreatePresence(Params.LocalUserId, Params.TargetUserId) };
+					UpdateUserPresence(Params.LocalAccountId, Params.TargetAccountId);
+					FQueryPresence::Result Result = { FindOrCreatePresence(Params.LocalAccountId, Params.TargetAccountId) };
 					InAsyncOp.SetResult(MoveTemp(Result));
 				}
 				else
@@ -154,9 +154,9 @@ TOnlineAsyncOpHandle<FQueryPresence> FPresenceEOS::QueryPresence(FQueryPresence:
 
 TOnlineResult<FGetCachedPresence> FPresenceEOS::GetCachedPresence(FGetCachedPresence::Params&& Params)
 {
-	if (TMap<FAccountId, TSharedRef<FUserPresence>>* PresenceList = PresenceLists.Find(Params.LocalUserId))
+	if (TMap<FAccountId, TSharedRef<FUserPresence>>* PresenceList = PresenceLists.Find(Params.LocalAccountId))
 	{
-		TSharedRef<FUserPresence>* PresencePtr = PresenceList->Find(Params.TargetUserId);
+		TSharedRef<FUserPresence>* PresencePtr = PresenceList->Find(Params.TargetAccountId);
 		if (PresencePtr)
 		{
 			FGetCachedPresence::Result Result = { *PresencePtr };
@@ -171,7 +171,7 @@ TOnlineAsyncOpHandle<FUpdatePresence> FPresenceEOS::UpdatePresence(FUpdatePresen
 {
 	TOnlineAsyncOpRef<FUpdatePresence> Op = GetOp<FUpdatePresence>(MoveTemp(InParams));
 	const FUpdatePresence::Params& Params = Op->GetParams();
-	if (!Services.Get<FAuthEOS>()->IsLoggedIn(Params.LocalUserId))
+	if (!Services.Get<FAuthEOS>()->IsLoggedIn(Params.LocalAccountId))
 	{
 		Op->SetError(Errors::NotLoggedIn());
 		return Op->GetHandle();
@@ -183,7 +183,7 @@ TOnlineAsyncOpHandle<FUpdatePresence> FPresenceEOS::UpdatePresence(FUpdatePresen
 		EOS_HPresenceModification ChangeHandle = nullptr;
 		EOS_Presence_CreatePresenceModificationOptions Options = { };
 		Options.ApiVersion = EOS_PRESENCE_CREATEPRESENCEMODIFICATION_API_LATEST;
-		Options.LocalUserId = GetEpicAccountIdChecked(Params.LocalUserId);
+		Options.LocalUserId = GetEpicAccountIdChecked(Params.LocalAccountId);
 		EOS_EResult CreatePresenceModificationResult = EOS_Presence_CreatePresenceModification(PresenceHandle, &Options, &ChangeHandle);
 		if (CreatePresenceModificationResult == EOS_EResult::EOS_Success)
 		{
@@ -221,7 +221,7 @@ TOnlineAsyncOpHandle<FUpdatePresence> FPresenceEOS::UpdatePresence(FUpdatePresen
 
 			// EOS needs to be specific on which fields are removed and which aren't, so grab the last presence and check which are deletions
 			TArray<FString> RemovedProperties;
-			TSharedRef<FUserPresence> LastUserPresence = FindOrCreatePresence(Params.LocalUserId, Params.LocalUserId);
+			TSharedRef<FUserPresence> LastUserPresence = FindOrCreatePresence(Params.LocalAccountId, Params.LocalAccountId);
 			for (const TPair<FString, FPresenceVariant>& Pair : LastUserPresence->Properties)
 			{
 				if(!Params.Presence->Properties.Contains(Pair.Key))
@@ -304,7 +304,7 @@ TOnlineAsyncOpHandle<FUpdatePresence> FPresenceEOS::UpdatePresence(FUpdatePresen
 
 			EOS_Presence_SetPresenceOptions SetPresenceOptions = { };
 			SetPresenceOptions.ApiVersion = EOS_PRESENCE_SETPRESENCE_API_LATEST;
-			SetPresenceOptions.LocalUserId = GetEpicAccountIdChecked(Params.LocalUserId);
+			SetPresenceOptions.LocalUserId = GetEpicAccountIdChecked(Params.LocalAccountId);
 			SetPresenceOptions.PresenceModificationHandle = ChangeHandle;
 
 			EOS_Async(EOS_Presence_SetPresence, PresenceHandle, SetPresenceOptions, MoveTemp(Promise));
@@ -325,7 +325,7 @@ TOnlineAsyncOpHandle<FUpdatePresence> FPresenceEOS::UpdatePresence(FUpdatePresen
 		{
 			// Update local presence
 			const FUpdatePresence::Params& Params = InAsyncOp.GetParams();
-			TSharedRef<FUserPresence> LocalUserPresence = FindOrCreatePresence(Params.LocalUserId, Params.LocalUserId);
+			TSharedRef<FUserPresence> LocalUserPresence = FindOrCreatePresence(Params.LocalAccountId, Params.LocalAccountId);
 			*LocalUserPresence = *Params.Presence;
 
 			InAsyncOp.SetResult(FUpdatePresence::Result());
@@ -353,7 +353,7 @@ TOnlineAsyncOpHandle<FPartialUpdatePresence> FPresenceEOS::PartialUpdatePresence
 	{
 		// Initialize
 		const FPartialUpdatePresence::Params& Params = Op->GetParams();
-		if (!Services.Get<FAuthEOS>()->IsLoggedIn(Params.LocalUserId))
+		if (!Services.Get<FAuthEOS>()->IsLoggedIn(Params.LocalAccountId))
 		{
 			Op->SetError(Errors::NotLoggedIn());
 			return Op->GetHandle();
@@ -366,7 +366,7 @@ TOnlineAsyncOpHandle<FPartialUpdatePresence> FPresenceEOS::PartialUpdatePresence
 				EOS_HPresenceModification ChangeHandle = nullptr;
 				EOS_Presence_CreatePresenceModificationOptions Options = { };
 				Options.ApiVersion = EOS_PRESENCE_CREATEPRESENCEMODIFICATION_API_LATEST;
-				Options.LocalUserId = GetEpicAccountIdChecked(Params.LocalUserId);
+				Options.LocalUserId = GetEpicAccountIdChecked(Params.LocalAccountId);
 				EOS_EResult CreatePresenceModificationResult = EOS_Presence_CreatePresenceModification(PresenceHandle, &Options, &ChangeHandle);
 				if (CreatePresenceModificationResult == EOS_EResult::EOS_Success)
 				{
@@ -489,7 +489,7 @@ TOnlineAsyncOpHandle<FPartialUpdatePresence> FPresenceEOS::PartialUpdatePresence
 
 					EOS_Presence_SetPresenceOptions SetPresenceOptions = { };
 					SetPresenceOptions.ApiVersion = EOS_PRESENCE_SETPRESENCE_API_LATEST;
-					SetPresenceOptions.LocalUserId = GetEpicAccountIdChecked(Params.LocalUserId);
+					SetPresenceOptions.LocalUserId = GetEpicAccountIdChecked(Params.LocalAccountId);
 					SetPresenceOptions.PresenceModificationHandle = ChangeHandle;
 					EOS_Async(EOS_Presence_SetPresence, PresenceHandle, SetPresenceOptions, MoveTemp(Promise));
 					EOS_PresenceModification_Release(ChangeHandle);
@@ -509,7 +509,7 @@ TOnlineAsyncOpHandle<FPartialUpdatePresence> FPresenceEOS::PartialUpdatePresence
 				{
 					// Update local presence
 					const FPartialUpdatePresence::Params& Params = InAsyncOp.GetParams();
-					TSharedRef<FUserPresence> LocalUserPresence = FindOrCreatePresence(Params.LocalUserId, Params.LocalUserId);
+					TSharedRef<FUserPresence> LocalUserPresence = FindOrCreatePresence(Params.LocalAccountId, Params.LocalAccountId);
 					if (Params.Mutations.Status.IsSet())
 					{
 						LocalUserPresence->Status = Params.Mutations.Status.GetValue();
@@ -548,32 +548,32 @@ TOnlineAsyncOpHandle<FPartialUpdatePresence> FPresenceEOS::PartialUpdatePresence
 }
 
 /** Get a user's presence, creating entries if missing */
-TSharedRef<FUserPresence> FPresenceEOS::FindOrCreatePresence(FAccountId LocalUserId, FAccountId PresenceUserId)
+TSharedRef<FUserPresence> FPresenceEOS::FindOrCreatePresence(FAccountId LocalAccountId, FAccountId PresenceAccountId)
 {
-	TMap<FAccountId, TSharedRef<FUserPresence>>& LocalUserPresenceList = PresenceLists.FindOrAdd(LocalUserId);
-	if (const TSharedRef<FUserPresence>* const ExistingPresence = LocalUserPresenceList.Find(PresenceUserId))
+	TMap<FAccountId, TSharedRef<FUserPresence>>& LocalUserPresenceList = PresenceLists.FindOrAdd(LocalAccountId);
+	if (const TSharedRef<FUserPresence>* const ExistingPresence = LocalUserPresenceList.Find(PresenceAccountId))
 	{
 		return *ExistingPresence;
 	}
 
 	TSharedRef<FUserPresence> UserPresence = MakeShared<FUserPresence>();
-	UserPresence->UserId = PresenceUserId;
-	LocalUserPresenceList.Emplace(PresenceUserId, UserPresence);
+	UserPresence->AccountId = PresenceAccountId;
+	LocalUserPresenceList.Emplace(PresenceAccountId, UserPresence);
 	return UserPresence;
 }
 
-void FPresenceEOS::UpdateUserPresence(FAccountId LocalUserId, FAccountId PresenceUserId)
+void FPresenceEOS::UpdateUserPresence(FAccountId LocalAccountId, FAccountId PresenceAccountId)
 {
 	bool bPresenceHasChanged = false;
-	TSharedRef<FUserPresence> UserPresence = FindOrCreatePresence(LocalUserId, PresenceUserId);
+	TSharedRef<FUserPresence> UserPresence = FindOrCreatePresence(LocalAccountId, PresenceAccountId);
 	// TODO:  Handle updates for local users.  Don't want to conflict with UpdatePresence calls
 
 	// Get presence from EOS
 	EOS_Presence_Info* PresenceInfo = nullptr;
 	EOS_Presence_CopyPresenceOptions Options = { };
 	Options.ApiVersion = EOS_PRESENCE_COPYPRESENCE_API_LATEST;
-	Options.LocalUserId = GetEpicAccountIdChecked(LocalUserId);
-	Options.TargetUserId = GetEpicAccountIdChecked(PresenceUserId);
+	Options.LocalUserId = GetEpicAccountIdChecked(LocalAccountId);
+	Options.TargetUserId = GetEpicAccountIdChecked(PresenceAccountId);
 	EOS_EResult CopyPresenceResult = EOS_Presence_CopyPresence(PresenceHandle, &Options, &PresenceInfo);
 	if (CopyPresenceResult == EOS_EResult::EOS_Success)
 	{
@@ -638,7 +638,7 @@ void FPresenceEOS::UpdateUserPresence(FAccountId LocalUserId, FAccountId Presenc
 
 	if (bPresenceHasChanged)
 	{
-		FPresenceUpdated PresenceUpdatedParams = { LocalUserId, UserPresence };
+		FPresenceUpdated PresenceUpdatedParams = { LocalAccountId, UserPresence };
 		OnPresenceUpdatedEvent.Broadcast(PresenceUpdatedParams);
 	}
 }
diff --git a/Engine/Plugins/Online/Experimental/OnlineServicesEOS/Source/Private/Online/Tests/OnlineIdEOSTests.cpp b/Engine/Plugins/Online/Experimental/OnlineServicesEOS/Source/Private/Online/Tests/OnlineIdEOSTests.cpp
--- a/Engine/Plugins/Online/Experimental/OnlineServicesEOS/Source/Private/Online/Tests/OnlineIdEOSTests.cpp
+++ b/Engine/Plugins/Online/Experimental/OnlineServicesEOS/Source/Private/Online/Tests/OnlineIdEOSTests.cpp
@@ -21,26 +21,26 @@ bool FAccountIdReplicationTest::RunTest(const FString& Parameters)
 
 		{
 			FOnlineAccountIdRegistryEOS Registry;
-			const FAccountId Handle = Registry.FindOrAddAccountId(EasId, EosId);
-			TArray<uint8> RepData = Registry.ToReplicationData(Handle);
-			const FAccountId Handle2 = Registry.FromReplicationData(RepData);
-			UTEST_EQUAL(TEXT(""), Handle, Handle2);
+			const FAccountId AccountId = Registry.FindOrAddAccountId(EasId, EosId);
+			TArray<uint8> RepData = Registry.ToReplicationData(AccountId);
+			const FAccountId AccountId2 = Registry.FromReplicationData(RepData);
+			UTEST_EQUAL(TEXT(""), AccountId, AccountId2);
 		}
 
 		{
 			FOnlineAccountIdRegistryEOS Registry;
-			const FAccountId Handle = Registry.FindOrAddAccountId(EasId, nullptr);
-			TArray<uint8> RepData = Registry.ToReplicationData(Handle);
-			const FAccountId Handle2 = Registry.FromReplicationData(RepData);
-			UTEST_EQUAL(TEXT(""), Handle, Handle2);
+			const FAccountId AccountId = Registry.FindOrAddAccountId(EasId, nullptr);
+			TArray<uint8> RepData = Registry.ToReplicationData(AccountId);
+			const FAccountId AccountId2 = Registry.FromReplicationData(RepData);
+			UTEST_EQUAL(TEXT(""), AccountId, AccountId2);
 		}
 
 		{
 			FOnlineAccountIdRegistryEOS Registry;
-			const FAccountId Handle = Registry.FindOrAddAccountId(nullptr, EosId);
-			TArray<uint8> RepData = Registry.ToReplicationData(Handle);
-			const FAccountId Handle2 = Registry.FromReplicationData(RepData);
-			UTEST_EQUAL(TEXT(""), Handle, Handle2);
+			const FAccountId AccountId = Registry.FindOrAddAccountId(nullptr, EosId);
+			TArray<uint8> RepData = Registry.ToReplicationData(AccountId);
+			const FAccountId AccountId2 = Registry.FromReplicationData(RepData);
+			UTEST_EQUAL(TEXT(""), AccountId, AccountId2);
 		}
 	}
 
diff --git a/Engine/Plugins/Online/Experimental/OnlineServicesEOS/Source/Private/Online/UserInfoEOS.cpp b/Engine/Plugins/Online/Experimental/OnlineServicesEOS/Source/Private/Online/UserInfoEOS.cpp
--- a/Engine/Plugins/Online/Experimental/OnlineServicesEOS/Source/Private/Online/UserInfoEOS.cpp
+++ b/Engine/Plugins/Online/Experimental/OnlineServicesEOS/Source/Private/Online/UserInfoEOS.cpp
@@ -30,26 +30,26 @@ TOnlineAsyncOpHandle<FQueryUserInfo> FUserInfoEOS::QueryUserInfo(FQueryUserInfo:
 	TOnlineAsyncOpRef<FQueryUserInfo> Op = GetJoinableOp<FQueryUserInfo>(MoveTemp(InParams));
 	const FQueryUserInfo::Params& Params = Op->GetParams();
 
-	if (Params.UserIds.IsEmpty())
+	if (Params.AccountIds.IsEmpty())
 	{
 		Op->SetError(Errors::InvalidParams());
 		return Op->GetHandle();
 	}
 
-	for (const FAccountId TargetUserId : Params.UserIds)
+	for (const FAccountId TargetAccountId : Params.AccountIds)
 	{
-		Op->Then([this, TargetUserId](TOnlineAsyncOp<FQueryUserInfo>& Op, TPromise<const EOS_UserInfo_QueryUserInfoCallbackInfo*>&& Promise)
+		Op->Then([this, TargetAccountId](TOnlineAsyncOp<FQueryUserInfo>& Op, TPromise<const EOS_UserInfo_QueryUserInfoCallbackInfo*>&& Promise)
 		{
 			const FQueryUserInfo::Params& Params = Op.GetParams();
 
-			if (!Services.Get<FAuthEOS>()->IsLoggedIn(Params.LocalUserId))
+			if (!Services.Get<FAuthEOS>()->IsLoggedIn(Params.LocalAccountId))
 			{
 				Op.SetError(Errors::NotLoggedIn());
 				Promise.EmplaceValue();
 				return;
 			}
 
-			const EOS_EpicAccountId TargetUserEasId = GetEpicAccountId(TargetUserId);
+			const EOS_EpicAccountId TargetUserEasId = GetEpicAccountId(TargetAccountId);
 			if (!EOS_EpicAccountId_IsValid(TargetUserEasId))
 			{
 				Op.SetError(Errors::InvalidParams());
@@ -60,7 +60,7 @@ TOnlineAsyncOpHandle<FQueryUserInfo> FUserInfoEOS::QueryUserInfo(FQueryUserInfo:
 			EOS_UserInfo_QueryUserInfoOptions QueryUserInfoOptions = {};
 			QueryUserInfoOptions.ApiVersion = EOS_USERINFO_QUERYUSERINFO_API_LATEST;
 			static_assert(EOS_USERINFO_QUERYUSERINFO_API_LATEST == 1, "EOS_UserInfo_QueryUserInfoOptions updated, check new fields");
-			QueryUserInfoOptions.LocalUserId = GetEpicAccountIdChecked(Params.LocalUserId);
+			QueryUserInfoOptions.LocalUserId = GetEpicAccountIdChecked(Params.LocalAccountId);
 			QueryUserInfoOptions.TargetUserId = TargetUserEasId;
 
 			EOS_Async(EOS_UserInfo_QueryUserInfo, UserInfoHandle, QueryUserInfoOptions, MoveTemp(Promise));
@@ -85,12 +85,12 @@ TOnlineAsyncOpHandle<FQueryUserInfo> FUserInfoEOS::QueryUserInfo(FQueryUserInfo:
 
 TOnlineResult<FGetUserInfo> FUserInfoEOS::GetUserInfo(FGetUserInfo::Params&& Params)
 {
-	if (!Services.Get<FAuthEOS>()->IsLoggedIn(Params.LocalUserId))
+	if (!Services.Get<FAuthEOS>()->IsLoggedIn(Params.LocalAccountId))
 	{
 		return TOnlineResult<FGetUserInfo>(Errors::NotLoggedIn());
 	}
 
-	const EOS_EpicAccountId TargetUserEasId = GetEpicAccountId(Params.UserId);
+	const EOS_EpicAccountId TargetUserEasId = GetEpicAccountId(Params.AccountId);
 	if (!EOS_EpicAccountId_IsValid(TargetUserEasId))
 	{
 		return TOnlineResult<FGetUserInfo>(Errors::InvalidParams());
@@ -99,7 +99,7 @@ TOnlineResult<FGetUserInfo> FUserInfoEOS::GetUserInfo(FGetUserInfo::Params&& Par
 	EOS_UserInfo_CopyUserInfoOptions Options;
 	Options.ApiVersion = EOS_USERINFO_COPYUSERINFO_API_LATEST;
 	static_assert(EOS_USERINFO_COPYUSERINFO_API_LATEST == 3, "EOS_UserInfo_CopyUserInfoOptions updated, check new fields");
-	Options.LocalUserId = GetEpicAccountIdChecked(Params.LocalUserId);
+	Options.LocalUserId = GetEpicAccountIdChecked(Params.LocalAccountId);
 	Options.TargetUserId = TargetUserEasId;
 
 	EOS_UserInfo* EosUserInfo = nullptr;
@@ -116,7 +116,7 @@ TOnlineResult<FGetUserInfo> FUserInfoEOS::GetUserInfo(FGetUserInfo::Params&& Par
 	}
 
 	TSharedRef<FUserInfo> UserInfo = MakeShared<FUserInfo>();
-	UserInfo->UserId = Params.UserId;
+	UserInfo->AccountId = Params.AccountId;
 	UserInfo->DisplayName = UTF8_TO_TCHAR(EosUserInfo->DisplayName);
 
 	return TOnlineResult<FGetUserInfo>({UserInfo});
diff --git a/Engine/Plugins/Online/Experimental/OnlineServicesEOS/Source/Public/Online/AuthEOS.h b/Engine/Plugins/Online/Experimental/OnlineServicesEOS/Source/Public/Online/AuthEOS.h
--- a/Engine/Plugins/Online/Experimental/OnlineServicesEOS/Source/Public/Online/AuthEOS.h
+++ b/Engine/Plugins/Online/Experimental/OnlineServicesEOS/Source/Public/Online/AuthEOS.h
@@ -29,20 +29,20 @@ public:
 	// End IAuth
 
 	// Begin FAuthEOSGS
-	virtual TFuture<FAccountId> ResolveAccountId(const FAccountId& LocalUserId, const EOS_ProductUserId ProductUserId) override;
-	virtual TFuture<TArray<FAccountId>> ResolveAccountIds(const FAccountId& LocalUserId, const TArray<EOS_ProductUserId>& ProductUserIds) override;
+	virtual TFuture<FAccountId> ResolveAccountId(const FAccountId& LocalAccountId, const EOS_ProductUserId ProductUserId) override;
+	virtual TFuture<TArray<FAccountId>> ResolveAccountIds(const FAccountId& LocalAccountId, const TArray<EOS_ProductUserId>& ProductUserIds) override;
 	virtual TFunction<TFuture<FAccountId>(FOnlineAsyncOp& InAsyncOp, const EOS_ProductUserId& ProductUserId)> ResolveProductIdFn() override;
 	virtual TFunction<TFuture<TArray<FAccountId>>(FOnlineAsyncOp& InAsyncOp, const TArray<EOS_ProductUserId>& ProductUserIds)> ResolveProductIdsFn() override;
 	// End FAuthEOSGS
 
-	TFuture<FAccountId> ResolveAccountId(const FAccountId& LocalUserId, const EOS_EpicAccountId EpicAccountId);
-	TFuture<TArray<FAccountId>> ResolveAccountIds(const FAccountId& LocalUserId, const TArray<EOS_EpicAccountId>& EpicAccountIds);
+	TFuture<FAccountId> ResolveAccountId(const FAccountId& LocalAccountId, const EOS_EpicAccountId EpicAccountId);
+	TFuture<TArray<FAccountId>> ResolveAccountIds(const FAccountId& LocalAccountId, const TArray<EOS_EpicAccountId>& EpicAccountIds);
 	TFunction<TFuture<FAccountId>(FOnlineAsyncOp& InAsyncOp, const EOS_EpicAccountId& EpicAccountId)> ResolveEpicIdFn();
 	TFunction<TFuture<TArray<FAccountId>>(FOnlineAsyncOp& InAsyncOp, const TArray<EOS_EpicAccountId>& EpicAccountIds)> ResolveEpicIdsFn();
 
 protected:
 	void ProcessSuccessfulLogin(TOnlineAsyncOp<FAuthLogin>& InAsyncOp);
-	void OnEASLoginStatusChanged(FAccountId LocalUserId, ELoginStatus PreviousStatus, ELoginStatus CurrentStatus);
+	void OnEASLoginStatusChanged(FAccountId LocalAccountId, ELoginStatus PreviousStatus, ELoginStatus CurrentStatus);
 
 	static FAccountId CreateAccountId(const EOS_EpicAccountId EpicAccountId, const EOS_ProductUserId ProductUserId);
 
diff --git a/Engine/Plugins/Online/Experimental/OnlineServicesEOS/Source/Public/Online/FriendsEOS.h b/Engine/Plugins/Online/Experimental/OnlineServicesEOS/Source/Public/Online/FriendsEOS.h
--- a/Engine/Plugins/Online/Experimental/OnlineServicesEOS/Source/Public/Online/FriendsEOS.h
+++ b/Engine/Plugins/Online/Experimental/OnlineServicesEOS/Source/Public/Online/FriendsEOS.h
@@ -27,7 +27,7 @@ public:
 	virtual TOnlineAsyncOpHandle<FAddFriend> AddFriend(FAddFriend::Params&& Params) override;
 
 protected:
-	void OnEOSFriendsUpdate(FAccountId LocalUserId, FAccountId FriendUserId, EOS_EFriendsStatus PreviousStatus, EOS_EFriendsStatus CurrentStatus);
+	void OnEOSFriendsUpdate(FAccountId LocalAccountId, FAccountId FriendAccountId, EOS_EFriendsStatus PreviousStatus, EOS_EFriendsStatus CurrentStatus);
 protected:
 	TMap<FAccountId, TMap<FAccountId, TSharedRef<FFriend>>> FriendsLists;
 	EOS_HFriends FriendsHandle = nullptr;
diff --git a/Engine/Plugins/Online/Experimental/OnlineServicesEOS/Source/Public/Online/PresenceEOS.h b/Engine/Plugins/Online/Experimental/OnlineServicesEOS/Source/Public/Online/PresenceEOS.h
--- a/Engine/Plugins/Online/Experimental/OnlineServicesEOS/Source/Public/Online/PresenceEOS.h
+++ b/Engine/Plugins/Online/Experimental/OnlineServicesEOS/Source/Public/Online/PresenceEOS.h
@@ -29,9 +29,9 @@ public:
 
 protected:
 	/** Get a user's presence, creating entries if missing */
-	TSharedRef<FUserPresence> FindOrCreatePresence(FAccountId LocalUserId, FAccountId PresenceUserId);
+	TSharedRef<FUserPresence> FindOrCreatePresence(FAccountId LocalAccountId, FAccountId PresenceAccountId);
 	/** Update a user's presence from EOS's current value */
-	void UpdateUserPresence(FAccountId LocalUserId, FAccountId PresenceUserId);
+	void UpdateUserPresence(FAccountId LocalAccountId, FAccountId PresenceAccountId);
 protected:
 	EOS_HPresence PresenceHandle = nullptr;
 
