Addition of Get/Set/ClearPresenceSession methods to Sessions interface

diff --git a/Engine/Plugins/Online/Experimental/OnlineServices/Source/OnlineServicesCommon/Private/Online/SessionsCommon.cpp b/Engine/Plugins/Online/Experimental/OnlineServices/Source/OnlineServicesCommon/Private/Online/SessionsCommon.cpp
--- a/Engine/Plugins/Online/Experimental/OnlineServices/Source/OnlineServicesCommon/Private/Online/SessionsCommon.cpp
+++ b/Engine/Plugins/Online/Experimental/OnlineServices/Source/OnlineServicesCommon/Private/Online/SessionsCommon.cpp
@@ -75,6 +75,33 @@ namespace UE::Online {
 			}
 		}
 
+	TOnlineResult<FGetPresenceSession> FSessionsCommon::GetPresenceSession(FGetPresenceSession::Params&& Params) const
+	{
+		if (const FOnlineSessionIdHandle* PresenceSessionId = PresenceSessionsUserMap.Find(Params.LocalUserId))
+		{
+			return TOnlineResult<FGetPresenceSession>({ AllSessionsById.FindChecked(*PresenceSessionId) });
+		}
+		else
+		{
+			return  TOnlineResult<FGetPresenceSession>({ Errors::InvalidState() });
+		}
+	}
+
+	TOnlineResult<FSetPresenceSession> FSessionsCommon::SetPresenceSession(FSetPresenceSession::Params&& Params)
+	{
+		FOnlineSessionIdHandle& PresenceSessionId = PresenceSessionsUserMap.FindOrAdd(Params.LocalUserId);
+		PresenceSessionId = Params.SessionId;
+
+		return TOnlineResult<FSetPresenceSession>(FSetPresenceSession::Result{ });
+	}
+
+	TOnlineResult<FClearPresenceSession> FSessionsCommon::ClearPresenceSession(FClearPresenceSession::Params&& Params)
+	{
+		PresenceSessionsUserMap.Remove(Params.LocalUserId);
+
+		return TOnlineResult<FClearPresenceSession>(FClearPresenceSession::Result{ });
+	}
+
 	TOnlineAsyncOpHandle<FCreateSession> FSessionsCommon::CreateSession(FCreateSession::Params&& Params)
 	{
 		TOnlineAsyncOpRef<FCreateSession> Operation = GetOp<FCreateSession>(MoveTemp(Params));
diff --git a/Engine/Plugins/Online/Experimental/OnlineServices/Source/OnlineServicesCommon/Public/Online/SessionsCommon.h b/Engine/Plugins/Online/Experimental/OnlineServices/Source/OnlineServicesCommon/Public/Online/SessionsCommon.h
--- a/Engine/Plugins/Online/Experimental/OnlineServices/Source/OnlineServicesCommon/Public/Online/SessionsCommon.h
+++ b/Engine/Plugins/Online/Experimental/OnlineServices/Source/OnlineServicesCommon/Public/Online/SessionsCommon.h
@@ -201,6 +201,9 @@ public:
 	virtual TOnlineResult<FGetAllSessions> GetAllSessions(FGetAllSessions::Params&& Params) const override;
 	virtual TOnlineResult<FGetSessionByName> GetSessionByName(FGetSessionByName::Params&& Params) const override;
 	virtual TOnlineResult<FGetSessionById> GetSessionById(FGetSessionById::Params&& Params) const override;
+	virtual TOnlineResult<FGetPresenceSession> GetPresenceSession(FGetPresenceSession::Params&& Params) const override;
+	virtual TOnlineResult<FSetPresenceSession> SetPresenceSession(FSetPresenceSession::Params&& Params) override;
+	virtual TOnlineResult<FClearPresenceSession> ClearPresenceSession(FClearPresenceSession::Params&& Params) override;
 	virtual TOnlineAsyncOpHandle<FCreateSession> CreateSession(FCreateSession::Params&& Params) override;
 	virtual TOnlineAsyncOpHandle<FUpdateSession> UpdateSession(FUpdateSession::Params&& Params) override;
 	virtual TOnlineAsyncOpHandle<FLeaveSession> LeaveSession(FLeaveSession::Params&& Params) override;
diff --git a/Engine/Plugins/Online/Experimental/OnlineServices/Source/OnlineServicesCommonEngineUtils/Private/Online/SessionsLAN.cpp b/Engine/Plugins/Online/Experimental/OnlineServices/Source/OnlineServicesCommonEngineUtils/Private/Online/SessionsLAN.cpp
--- a/Engine/Plugins/Online/Experimental/OnlineServices/Source/OnlineServicesCommonEngineUtils/Private/Online/SessionsLAN.cpp
+++ b/Engine/Plugins/Online/Experimental/OnlineServices/Source/OnlineServicesCommonEngineUtils/Private/Online/SessionsLAN.cpp
@@ -168,8 +168,7 @@ TOnlineAsyncOpHandle<FCreateSession> FSessionsLAN::CreateSession(FCreateSession:
 
 		if (NewSessionLANRef->SessionSettings.bPresenceEnabled)
 		{
-			FOnlineSessionIdHandle& PresenceSessionId = PresenceSessionsUserMap.FindOrAdd(OpParams.LocalAccountId);
-			PresenceSessionId = NewSessionLANRef->SessionId;
+			SetPresenceSession({ OpParams.LocalAccountId, NewSessionLANRef->SessionId });
 		}
 
 		Op.SetResult({ });
@@ -312,8 +311,7 @@ TOnlineAsyncOpHandle<FJoinSession> FSessionsLAN::JoinSession(FJoinSession::Param
 
 		if (FoundSession->GetSessionSettings().bPresenceEnabled)
 		{
-			FOnlineSessionIdHandle& PresenceSessionId = PresenceSessionsUserMap.FindOrAdd(OpParams.LocalAccountId);
-			PresenceSessionId = FoundSession->GetSessionId();
+			SetPresenceSession({ OpParams.LocalAccountId, FoundSession->GetSessionId() });
 		}
 
 		Op.SetResult(FJoinSession::Result{ });
@@ -351,11 +349,16 @@ TOnlineAsyncOpHandle<FLeaveSession> FSessionsLAN::LeaveSession(FLeaveSession::Pa
 		{
 			TSharedRef<const ISession> FoundSession = GetSessionByNameResult.GetOkValue().Session;
 
+			if (FoundSession->GetOwnerAccountId() == OpParams.LocalAccountId && FoundSession->GetSessionSettings().JoinPolicy == ESessionJoinPolicy::Public)
+			{
+				StopLANSession();
+			}
+
 			NamedSessionUserMap.FindChecked(OpParams.LocalAccountId).Remove(OpParams.SessionName);
 
 			if (FoundSession->GetSessionSettings().bPresenceEnabled)
 			{
-				PresenceSessionsUserMap.Remove(OpParams.LocalAccountId);
+				ClearPresenceSession({ OpParams.LocalAccountId });
 			}
 
 			ClearSessionByName(OpParams.SessionName);
diff --git a/Engine/Plugins/Online/Experimental/OnlineServices/Source/OnlineServicesInterface/Public/Online/Sessions.h b/Engine/Plugins/Online/Experimental/OnlineServices/Source/OnlineServicesInterface/Public/Online/Sessions.h
--- a/Engine/Plugins/Online/Experimental/OnlineServices/Source/OnlineServicesInterface/Public/Online/Sessions.h
+++ b/Engine/Plugins/Online/Experimental/OnlineServices/Source/OnlineServicesInterface/Public/Online/Sessions.h
@@ -256,6 +256,51 @@ struct FGetSessionById
 	};
 };
 
+struct FGetPresenceSession
+{
+	static constexpr TCHAR Name[] = TEXT("GetPresenceSession");
+
+	struct Params
+	{
+		FAccountId LocalUserId;
+	};
+
+	struct Result
+	{
+		TSharedPtr<const ISession> Session;
+	};
+};
+
+struct FSetPresenceSession
+{
+	static constexpr TCHAR Name[] = TEXT("SetPresenceSession");
+
+	struct Params
+	{
+		FAccountId LocalUserId;
+
+		FOnlineSessionIdHandle SessionId;
+	};
+
+	struct Result
+	{
+	};
+};
+
+struct FClearPresenceSession
+{
+	static constexpr TCHAR Name[] = TEXT("ClearPresenceSession");
+
+	struct Params
+	{
+		FAccountId LocalUserId;
+	};
+
+	struct Result
+	{
+	};
+};
+
 struct FCreateSession
 {
 	static constexpr TCHAR Name[] = TEXT("CreateSession");
@@ -656,6 +701,30 @@ public:
 	 */
 	virtual TOnlineResult<FGetSessionById> GetSessionById(FGetSessionById::Params&& Params) const = 0;
 
+	/**
+	 * Get the session set as presence session for the user.
+	 *
+	 * @params Parameters for the GetPresenceSession call
+	 * return
+	 */
+	virtual TOnlineResult<FGetPresenceSession> GetPresenceSession(FGetPresenceSession::Params&& Params) const = 0;
+
+	/**
+	 * Sets the session with the given id as the presence session for the user.
+	 *
+	 * @params Parameters for the SetPresenceSession call
+	 * return
+	 */
+	virtual TOnlineResult<FSetPresenceSession> SetPresenceSession(FSetPresenceSession::Params&& Params) = 0;
+
+	/**
+	 * Clears the presence session for the user. If no presence session is set, GetPresenceSession will return an error.
+	 *
+	 * @params Parameters for the ClearPresenceSession call
+	 * return
+	 */
+	virtual TOnlineResult<FClearPresenceSession> ClearPresenceSession(FClearPresenceSession::Params&& Params) = 0;
+
 	/**
 	 * Create and join a new session.
 	 *
@@ -912,6 +981,29 @@ BEGIN_ONLINE_STRUCT_META(FGetSessionById::Result)
 	ONLINE_STRUCT_FIELD(FGetSessionById::Result, Session)
 END_ONLINE_STRUCT_META()
 
+BEGIN_ONLINE_STRUCT_META(FGetPresenceSession::Params)
+	ONLINE_STRUCT_FIELD(FGetPresenceSession::Params, LocalUserId)
+END_ONLINE_STRUCT_META()
+
+BEGIN_ONLINE_STRUCT_META(FGetPresenceSession::Result)
+	ONLINE_STRUCT_FIELD(FGetPresenceSession::Result, Session)
+END_ONLINE_STRUCT_META()
+
+BEGIN_ONLINE_STRUCT_META(FSetPresenceSession::Params)
+	ONLINE_STRUCT_FIELD(FSetPresenceSession::Params, LocalUserId),
+	ONLINE_STRUCT_FIELD(FSetPresenceSession::Params, SessionId)
+END_ONLINE_STRUCT_META()
+
+BEGIN_ONLINE_STRUCT_META(FSetPresenceSession::Result)
+END_ONLINE_STRUCT_META()
+
+BEGIN_ONLINE_STRUCT_META(FClearPresenceSession::Params)
+	ONLINE_STRUCT_FIELD(FClearPresenceSession::Params, LocalUserId)
+END_ONLINE_STRUCT_META()
+
+BEGIN_ONLINE_STRUCT_META(FClearPresenceSession::Result)
+END_ONLINE_STRUCT_META()
+
 BEGIN_ONLINE_STRUCT_META(FCreateSession::Params)
 	ONLINE_STRUCT_FIELD(FCreateSession::Params, LocalAccountId),
 	ONLINE_STRUCT_FIELD(FCreateSession::Params, SessionName),
diff --git a/Engine/Plugins/Online/Experimental/OnlineServicesEOSGS/Source/Private/Online/SessionsEOSGS.cpp b/Engine/Plugins/Online/Experimental/OnlineServicesEOSGS/Source/Private/Online/SessionsEOSGS.cpp
--- a/Engine/Plugins/Online/Experimental/OnlineServicesEOSGS/Source/Private/Online/SessionsEOSGS.cpp
+++ b/Engine/Plugins/Online/Experimental/OnlineServicesEOSGS/Source/Private/Online/SessionsEOSGS.cpp
@@ -503,8 +503,7 @@ TOnlineAsyncOpHandle<FCreateSession> FSessionsEOSGS::CreateSession(FCreateSessio
 
 			if (NewSessionEOSGSRef->SessionSettings.bPresenceEnabled)
 			{
-				FOnlineSessionIdHandle& PresenceSessionId = PresenceSessionsUserMap.FindOrAdd(OpParams.LocalAccountId);
-				PresenceSessionId = NewSessionEOSGSRef->SessionId;
+				SetPresenceSession({ OpParams.LocalAccountId, NewSessionEOSGSRef->SessionId });
 			}
 
 			Op.SetResult({ });
@@ -1153,7 +1152,7 @@ TOnlineAsyncOpHandle<FLeaveSession> FSessionsEOSGS::LeaveSession(FLeaveSession::
 
 			if (FoundSession->GetSessionSettings().bPresenceEnabled)
 			{
-				PresenceSessionsUserMap.Remove(OpParams.LocalAccountId);
+				ClearPresenceSession({ OpParams.LocalAccountId });
 			}
 
 			ClearSessionByName(OpParams.SessionName);
@@ -1558,8 +1557,7 @@ TOnlineAsyncOpHandle<FJoinSession> FSessionsEOSGS::JoinSession(FJoinSession::Par
 
 		if (FoundSession->GetSessionSettings().bPresenceEnabled)
 		{
-			FOnlineSessionIdHandle& PresenceSessionId = PresenceSessionsUserMap.FindOrAdd(OpParams.LocalAccountId);
-			PresenceSessionId = FoundSession->GetSessionId();
+			SetPresenceSession({ OpParams.LocalAccountId, FoundSession->GetSessionId() });
 		}
 
 		// After successfully joining a session, we'll remove all related invites if any are found
diff --git a/Engine/Plugins/Online/Experimental/OnlineServicesOSSAdapter/Source/Private/Online/SessionsOSSAdapter.cpp b/Engine/Plugins/Online/Experimental/OnlineServicesOSSAdapter/Source/Private/Online/SessionsOSSAdapter.cpp
--- a/Engine/Plugins/Online/Experimental/OnlineServicesOSSAdapter/Source/Private/Online/SessionsOSSAdapter.cpp
+++ b/Engine/Plugins/Online/Experimental/OnlineServicesOSSAdapter/Source/Private/Online/SessionsOSSAdapter.cpp
@@ -463,6 +463,11 @@ TOnlineAsyncOpHandle<FCreateSession> FSessionsOSSAdapter::CreateSession(FCreateS
 
 				NamedSessionUserMap.FindOrAdd(OpParams.LocalAccountId).AddUnique(OpParams.SessionName);
 
+				if (V2Session->SessionSettings.bPresenceEnabled)
+				{
+					SetPresenceSession({ OpParams.LocalAccountId, V2Session->SessionId });
+				}
+
 				Op->SetResult({ });
 			}
 		});
@@ -583,7 +588,7 @@ TOnlineAsyncOpHandle<FLeaveSession> FSessionsOSSAdapter::LeaveSession(FLeaveSess
 
 					if (FoundSession->GetSessionSettings().bPresenceEnabled)
 					{
-						PresenceSessionsUserMap.Remove(OpParams.LocalAccountId);
+						ClearPresenceSession({ OpParams.LocalAccountId });
 					}
 
 					ClearSessionByName(OpParams.SessionName);
@@ -884,6 +889,11 @@ TOnlineAsyncOpHandle<FStartMatchmaking> FSessionsOSSAdapter::StartMatchmaking(FS
 
 				NamedSessionUserMap.FindOrAdd(OpParams.LocalAccountId).AddUnique(OpParams.SessionName);
 
+				if (V2Session->SessionSettings.bPresenceEnabled)
+				{
+					SetPresenceSession({ OpParams.LocalAccountId, V2Session->SessionId });
+				}
+
 				Op->SetResult({ });
 
 				FSessionJoined Event{ { Op->GetParams().LocalAccountId }, V2Session->SessionId };
@@ -963,8 +973,7 @@ TOnlineAsyncOpHandle<FJoinSession> FSessionsOSSAdapter::JoinSession(FJoinSession
 
 				if (FoundSession->GetSessionSettings().bPresenceEnabled)
 				{
-					FOnlineSessionIdHandle& PresenceSessionId = PresenceSessionsUserMap.FindOrAdd(OpParams.LocalAccountId);
-					PresenceSessionId = FoundSession->GetSessionId();
+					SetPresenceSession({ OpParams.LocalAccountId, FoundSession->GetSessionId() });
 				}
 
 				// After successfully joining a session, we'll remove all related invites if any are found
