Removal of multiple user support in Sessions interface

diff --git a/Engine/Plugins/Online/Experimental/OnlineServices/Source/OnlineServicesCommon/Private/Online/SessionsCommon.cpp b/Engine/Plugins/Online/Experimental/OnlineServices/Source/OnlineServicesCommon/Private/Online/SessionsCommon.cpp
--- a/Engine/Plugins/Online/Experimental/OnlineServices/Source/OnlineServicesCommon/Private/Online/SessionsCommon.cpp
+++ b/Engine/Plugins/Online/Experimental/OnlineServices/Source/OnlineServicesCommon/Private/Online/SessionsCommon.cpp
@@ -77,7 +77,7 @@ namespace UE::Online {
 
 	TOnlineResult<FGetPresenceSession> FSessionsCommon::GetPresenceSession(FGetPresenceSession::Params&& Params) const
 	{
-		if (const FOnlineSessionIdHandle* PresenceSessionId = PresenceSessionsUserMap.Find(Params.LocalUserId))
+		if (const FOnlineSessionIdHandle* PresenceSessionId = PresenceSessionsUserMap.Find(Params.LocalAccountId))
 		{
 			return TOnlineResult<FGetPresenceSession>({ AllSessionsById.FindChecked(*PresenceSessionId) });
 		}
@@ -89,7 +89,7 @@ namespace UE::Online {
 
 	TOnlineResult<FSetPresenceSession> FSessionsCommon::SetPresenceSession(FSetPresenceSession::Params&& Params)
 	{
-		FOnlineSessionIdHandle& PresenceSessionId = PresenceSessionsUserMap.FindOrAdd(Params.LocalUserId);
+		FOnlineSessionIdHandle& PresenceSessionId = PresenceSessionsUserMap.FindOrAdd(Params.LocalAccountId);
 		PresenceSessionId = Params.SessionId;
 
 		return TOnlineResult<FSetPresenceSession>(FSetPresenceSession::Result{ });
@@ -97,7 +97,7 @@ namespace UE::Online {
 
 	TOnlineResult<FClearPresenceSession> FSessionsCommon::ClearPresenceSession(FClearPresenceSession::Params&& Params)
 	{
-		PresenceSessionsUserMap.Remove(Params.LocalUserId);
+		PresenceSessionsUserMap.Remove(Params.LocalAccountId);
 
 		return TOnlineResult<FClearPresenceSession>(FClearPresenceSession::Result{ });
 	}
@@ -144,22 +144,22 @@ namespace UE::Online {
 		return Operation->GetHandle();
 	}
 
-	TOnlineAsyncOpHandle<FAddSessionMembers> FSessionsCommon::AddSessionMembers(FAddSessionMembers::Params&& Params)
+	TOnlineAsyncOpHandle<FAddSessionMember> FSessionsCommon::AddSessionMember(FAddSessionMember::Params&& Params)
 	{
-		TOnlineAsyncOpRef<FAddSessionMembers> Op = GetOp<FAddSessionMembers>(MoveTemp(Params));
+		TOnlineAsyncOpRef<FAddSessionMember> Op = GetOp<FAddSessionMember>(MoveTemp(Params));
 
-		Op->Then([this](TOnlineAsyncOp<FAddSessionMembers>& Op) mutable
+		Op->Then([this](TOnlineAsyncOp<FAddSessionMember>& Op) mutable
 		{
-			const FAddSessionMembers::Params& OpParams = Op.GetParams();
+			const FAddSessionMember::Params& OpParams = Op.GetParams();
 
-			FOnlineError StateCheck = CheckAddSessionMembersState(OpParams);
+			FOnlineError StateCheck = CheckAddSessionMemberState(OpParams);
 			if (StateCheck != Errors::Success())
 			{
 				Op.SetError(MoveTemp(StateCheck));
 				return;
 			}
 
-			TOnlineResult<FAddSessionMembers> Result = AddSessionMembersImpl(OpParams);
+			TOnlineResult<FAddSessionMember> Result = AddSessionMemberImpl(OpParams);
 			if (Result.IsOk())
 			{
 				Op.SetResult(MoveTemp(Result.GetOkValue()));
@@ -174,22 +174,22 @@ namespace UE::Online {
 		return Op->GetHandle();
 	}
 
-	TOnlineAsyncOpHandle<FRemoveSessionMembers> FSessionsCommon::RemoveSessionMembers(FRemoveSessionMembers::Params&& Params)
+	TOnlineAsyncOpHandle<FRemoveSessionMember> FSessionsCommon::RemoveSessionMember(FRemoveSessionMember::Params&& Params)
 	{
-		TOnlineAsyncOpRef<FRemoveSessionMembers> Op = GetOp<FRemoveSessionMembers>(MoveTemp(Params));
+		TOnlineAsyncOpRef<FRemoveSessionMember> Op = GetOp<FRemoveSessionMember>(MoveTemp(Params));
 
-		Op->Then([this](TOnlineAsyncOp<FRemoveSessionMembers>& Op) mutable
+		Op->Then([this](TOnlineAsyncOp<FRemoveSessionMember>& Op) mutable
 		{
-			const FRemoveSessionMembers::Params& OpParams = Op.GetParams();
+			const FRemoveSessionMember::Params& OpParams = Op.GetParams();
 
-			FOnlineError StateCheck = CheckRemoveSessionMembersState(OpParams);
+			FOnlineError StateCheck = CheckRemoveSessionMemberState(OpParams);
 			if (StateCheck != Errors::Success())
 			{
 				Op.SetError(MoveTemp(StateCheck));
 				return;
 			}
 
-			TOnlineResult<FRemoveSessionMembers> Result = RemoveSessionMembersImpl(OpParams);
+			TOnlineResult<FRemoveSessionMember> Result = RemoveSessionMemberImpl(OpParams);
 			if (Result.IsOk())
 			{
 				Op.SetResult(MoveTemp(Result.GetOkValue()));
@@ -398,121 +398,115 @@ namespace UE::Online {
 		AllSessionsById.Remove(SessionId);
 	}
 
-	TOnlineResult<FAddSessionMembers> FSessionsCommon::AddSessionMembersImpl(const FAddSessionMembers::Params& Params)
+	TOnlineResult<FAddSessionMember> FSessionsCommon::AddSessionMemberImpl(const FAddSessionMember::Params& Params)
 	{
 		TOnlineResult<FGetMutableSessionByName> GetMutableSessionByNameResult = GetMutableSessionByName({ Params.SessionName });
 		if (GetMutableSessionByNameResult.IsError())
 		{
-			return TOnlineResult<FAddSessionMembers>(GetMutableSessionByNameResult.GetErrorValue());
+			return TOnlineResult<FAddSessionMember>(GetMutableSessionByNameResult.GetErrorValue());
 		}
 		TSharedRef<FSessionCommon> FoundSession = GetMutableSessionByNameResult.GetOkValue().Session;
 
 		FSessionSettings& SessionSettings = FoundSession->SessionSettings;
 
-		for (const TPair<FAccountId, FSessionMember>& Entry : Params.NewSessionMembers)
-		{
-			const FAccountId& SessionMemberId = Entry.Key;
-			const FSessionMember& SessionMember = Entry.Value;
-
-			// If the user was already registered in the session, we'll update it and reserve a slot for it
-			if (FRegisteredPlayer* RegisteredPlayer = SessionSettings.RegisteredPlayers.Find(SessionMemberId))
-			{
-				RegisteredPlayer->bIsInSession = true;
-
-				if (!RegisteredPlayer->bHasReservedSlot)
-				{
-					bool bSlotReserved = false;
-
-					if (SessionSettings.NumOpenPublicConnections > 0)
-					{
-						SessionSettings.NumOpenPublicConnections--;
-						bSlotReserved = true;
-					}
-					else
-					{
-						return TOnlineResult<FAddSessionMembers>(Errors::InvalidState());
-					}
+		const FAccountId& SessionMemberId = Params.LocalAccountId;
+		const FSessionMember& SessionMember = Params.NewSessionMember;
 
-					// TODO: Discuss usage for Private Connections
+		// If the user was already registered in the session, we'll update it and reserve a slot for it
+		if (FRegisteredPlayer* RegisteredPlayer = SessionSettings.RegisteredPlayers.Find(SessionMemberId))
+		{
+			RegisteredPlayer->bIsInSession = true;
 
-					RegisteredPlayer->bHasReservedSlot = bSlotReserved;
-				}
-			}
-			else if (Params.bRegisterPlayers) // If it wasn't and the parameters dictate it, we'll register them with a reserved slot
+			if (!RegisteredPlayer->bHasReservedSlot)
 			{
-				FRegisterPlayers::Params RegisterPlayerParams = { Params.LocalAccountId, Params.SessionName, { SessionMemberId }, true };
+				bool bSlotReserved = false;
 
-				TOnlineResult<FRegisterPlayers> Result = RegisterPlayersImpl(RegisterPlayerParams);
-				if (Result.IsError())
-				{
-					return TOnlineResult<FAddSessionMembers>(MoveTemp(Result.GetErrorValue()));
-				}
-			}
-			else // Even if the new session member is not a registered player, we'll still need to decrease the number of open slots
-			{
 				if (SessionSettings.NumOpenPublicConnections > 0)
 				{
 					SessionSettings.NumOpenPublicConnections--;
+					bSlotReserved = true;
 				}
 				else
 				{
-					return TOnlineResult<FAddSessionMembers>(Errors::InvalidState());
+					return TOnlineResult<FAddSessionMember>(Errors::InvalidState());
 				}
 
 				// TODO: Discuss usage for Private Connections
+
+				RegisteredPlayer->bHasReservedSlot = bSlotReserved;
+			}
+		}
+		else if (Params.bRegisterPlayers) // If it wasn't and the parameters dictate it, we'll register them with a reserved slot
+		{
+			FRegisterPlayers::Params RegisterPlayerParams = { Params.LocalAccountId, Params.SessionName, { SessionMemberId }, true };
+
+			TOnlineResult<FRegisterPlayers> Result = RegisterPlayersImpl(RegisterPlayerParams);
+			if (Result.IsError())
+			{
+				return TOnlineResult<FAddSessionMember>(MoveTemp(Result.GetErrorValue()));
+			}
+		}
+		else // Even if the new session member is not a registered player, we'll still need to decrease the number of open slots
+		{
+			if (SessionSettings.NumOpenPublicConnections > 0)
+			{
+				SessionSettings.NumOpenPublicConnections--;
+			}
+			else
+			{
+				return TOnlineResult<FAddSessionMember>(Errors::InvalidState());
 			}
 
-			SessionSettings.SessionMembers.Emplace(SessionMemberId, SessionMember);
+			// TODO: Discuss usage for Private Connections
 		}
 
-		return TOnlineResult<FAddSessionMembers>();
+		SessionSettings.SessionMembers.Emplace(SessionMemberId, SessionMember);
+
+		return TOnlineResult<FAddSessionMember>();
 	}
 
-	TOnlineResult<FRemoveSessionMembers> FSessionsCommon::RemoveSessionMembersImpl(const FRemoveSessionMembers::Params& Params)
+	TOnlineResult<FRemoveSessionMember> FSessionsCommon::RemoveSessionMemberImpl(const FRemoveSessionMember::Params& Params)
 	{
 		TOnlineResult<FGetMutableSessionByName> GetMutableSessionByNameResult = GetMutableSessionByName({ Params.SessionName });
 		if (GetMutableSessionByNameResult.IsError())
 		{
-			return TOnlineResult<FRemoveSessionMembers>(GetMutableSessionByNameResult.GetErrorValue());
+			return TOnlineResult<FRemoveSessionMember>(GetMutableSessionByNameResult.GetErrorValue());
 		}
 		TSharedRef<FSessionCommon> FoundSession = GetMutableSessionByNameResult.GetOkValue().Session;
 
 		FSessionSettings& SessionSettings = FoundSession->SessionSettings;
-
-		for (const FAccountId& SessionMemberId : Params.SessionMemberIds)
+		
+		if (Params.bUnregisterPlayers) // If the parameters dictate it, well unregister the player as well
 		{
-			if (Params.bUnregisterPlayers) // If the parameters dictate it, well unregister the player as well
-			{
-				FUnregisterPlayers::Params UnregisterPlayersParams = { Params.LocalAccountId, Params.SessionName, { SessionMemberId }, false };
+			FUnregisterPlayers::Params UnregisterPlayersParams = { Params.LocalAccountId, Params.SessionName, { Params.LocalAccountId }, false };
 
-				TOnlineResult<FUnregisterPlayers> Result = UnregisterPlayersImpl(UnregisterPlayersParams);
-				if (Result.IsError())
-				{
-					return TOnlineResult<FRemoveSessionMembers>(MoveTemp(Result.GetErrorValue()));
-				}
+			TOnlineResult<FUnregisterPlayers> Result = UnregisterPlayersImpl(UnregisterPlayersParams);
+			if (Result.IsError())
+			{
+				return TOnlineResult<FRemoveSessionMember>(MoveTemp(Result.GetErrorValue()));
 			}
-			else if (FRegisteredPlayer* RegisteredPlayer = SessionSettings.RegisteredPlayers.Find(SessionMemberId)) // If the removed session member is still registered, we'll update it and keep their reserved slot
+		}
+		else if (FRegisteredPlayer* RegisteredPlayer = SessionSettings.RegisteredPlayers.Find(Params.LocalAccountId)) // If the removed session member is still registered, we'll update it and keep their reserved slot
+		{
+			RegisteredPlayer->bIsInSession = false;
+		}
+		else  // If the removed session member is not a registered player with a reserved slot, we'll increase the number of open slots
+		{
+			if (SessionSettings.NumOpenPublicConnections < SessionSettings.NumMaxPublicConnections)
 			{
-				RegisteredPlayer->bIsInSession = false;
+				SessionSettings.NumOpenPublicConnections++;
 			}
-			else  // If the removed session member is not a registered player with a reserved slot, we'll increase the number of open slots
+			else
 			{
-				if (SessionSettings.NumOpenPublicConnections < SessionSettings.NumMaxPublicConnections)
-				{
-					SessionSettings.NumOpenPublicConnections++;
-				}
-				else
-				{
-					return TOnlineResult<FRemoveSessionMembers>(Errors::InvalidState());
-				}
-
-				// TODO: Discuss usage for Private Connections
+				return TOnlineResult<FRemoveSessionMember>(Errors::InvalidState());
 			}
 
-			SessionSettings.SessionMembers.Remove(SessionMemberId);
+			// TODO: Discuss usage for Private Connections
 		}
 
-		return TOnlineResult<FRemoveSessionMembers>();
+		SessionSettings.SessionMembers.Remove(Params.LocalAccountId);
+
+		return TOnlineResult<FRemoveSessionMember>();
 	}
 
 	TOnlineResult<FRegisterPlayers> FSessionsCommon::RegisterPlayersImpl(const FRegisterPlayers::Params& Params)
@@ -591,9 +585,9 @@ namespace UE::Online {
 
 				if (RegisteredPlayer->bIsInSession && Params.bRemoveUnregisteredPlayers)
 				{
-					FRemoveSessionMembers::Params RemoveSessionMembersParams = { User, Params.SessionName, Params.TargetUsers, false };
+					FRemoveSessionMember::Params RemoveSessionMemberParams = { User, Params.SessionName, false };
 
-					TOnlineResult<FRemoveSessionMembers> Result = RemoveSessionMembersImpl(RemoveSessionMembersParams);
+					TOnlineResult<FRemoveSessionMember> Result = RemoveSessionMemberImpl(RemoveSessionMemberParams);
 					if (Result.IsError())
 					{
 						return TOnlineResult<FUnregisterPlayers>(MoveTemp(Result.GetErrorValue()));
@@ -629,13 +623,6 @@ namespace UE::Online {
 			return Errors::InvalidParams();
 		}
 
-		if (!Params.LocalAccounts.Contains(Params.LocalAccountId))
-		{
-			UE_LOG(LogTemp, Warning, TEXT("[FSessionsCommon::CheckCreateSessionParams] Could not create session with name [%s]. Missing Session Member info for user [%s]"), *Params.SessionName.ToString(), *ToLogString(Params.LocalAccountId));
-
-			return Errors::InvalidParams();
-		}
-
 		if (Params.SessionSettings.NumMaxPrivateConnections == 0 && Params.SessionSettings.NumMaxPublicConnections == 0)
 		{
 			UE_LOG(LogTemp, Warning, TEXT("[FSessionsCommon::CheckCreateSessionParams] Could not create session with name [%s] with no valid NumMaxPrivateConnections [%d] or NumMaxPublicConnections [%d]"), *Params.SessionName.ToString(), Params.SessionSettings.NumMaxPrivateConnections, Params.SessionSettings.NumMaxPublicConnections);
@@ -675,19 +662,11 @@ namespace UE::Online {
 
 		// User login check for all local users
 		IAuthPtr Auth = Services.GetAuthInterface();
-
-		TArray<FAccountId> LocalAccountIds;
-		LocalAccountIds.Reserve(Params.LocalAccounts.Num());
-		Params.LocalAccounts.GenerateKeyArray(LocalAccountIds);
-
-		for (const FAccountId& LocalAccountId : LocalAccountIds)
+		if (!Auth->IsLoggedIn(Params.LocalAccountId))
 		{
-			if (!Auth->IsLoggedIn(LocalAccountId))
-			{
-				UE_LOG(LogTemp, Warning, TEXT("[FSessionsCommon::CheckCreateSessionState] Could not create session with user [%s] not logged in"), *ToLogString(LocalAccountId));
+			UE_LOG(LogTemp, Warning, TEXT("[FSessionsCommon::CheckCreateSessionState] Could not create session with user [%s] not logged in"), *ToLogString(Params.LocalAccountId));
 
-				return Errors::InvalidUser();
-			}
+			return Errors::InvalidUser();
 		}
 
 		return Errors::Success();
@@ -778,19 +757,11 @@ namespace UE::Online {
 
 		// User login check for all local users
 		IAuthPtr Auth = Services.GetAuthInterface();
-
-		TArray<FAccountId> LocalAccountIds;
-		LocalAccountIds.Reserve(Params.LocalAccounts.Num());
-		Params.LocalAccounts.GenerateKeyArray(LocalAccountIds);
-
-		for (const FAccountId& LocalAccountId : LocalAccountIds)
+		if (!Auth->IsLoggedIn(Params.LocalAccountId))
 		{
-			if (!Auth->IsLoggedIn(LocalAccountId))
-			{
-				UE_LOG(LogTemp, Warning, TEXT("[FSessionsCommon::CheckStartMatchmakingState] Could not join session with user [%s] not logged in"), *ToLogString(LocalAccountId));
+			UE_LOG(LogTemp, Warning, TEXT("[FSessionsCommon::CheckStartMatchmakingState] Could not join session with user [%s] not logged in"), *ToLogString(Params.LocalAccountId));
 
-				return Errors::InvalidUser();
-			}
+			return Errors::InvalidUser();
 		}
 
 		// TODO: Check that only one session has bUsesPresence set
@@ -807,13 +778,6 @@ namespace UE::Online {
 			return Errors::InvalidParams();
 		}
 
-		if (!Params.LocalAccounts.Contains(Params.LocalAccountId))
-		{
-			UE_LOG(LogTemp, Warning, TEXT("[FSessionsCommon::CheckJoinSessionParams] Could not join session with name [%s]. Missing Session Member info for user [%s]"), *ToLogString(Params.LocalAccountId));
-
-			return Errors::InvalidParams();
-		}
-
 		if (!Params.SessionId.IsValid())
 		{
 			UE_LOG(LogTemp, Warning, TEXT("[FSessionsCommon::CheckJoinSessionParams] Could not join session with name [%s]. SessionId [%s] not valid"), *ToLogString(Params.SessionId));
@@ -844,19 +808,11 @@ namespace UE::Online {
 
 		// User login check for all local users
 		IAuthPtr Auth = Services.GetAuthInterface();
-
-		TArray<FAccountId> LocalAccountIds;
-		LocalAccountIds.Reserve(Params.LocalAccounts.Num());
-		Params.LocalAccounts.GenerateKeyArray(LocalAccountIds);
-
-		for (const FAccountId& LocalAccountId : LocalAccountIds)
+		if(!Auth->IsLoggedIn(Params.LocalAccountId))
 		{
-			if(!Auth->IsLoggedIn(LocalAccountId))
-			{
-				UE_LOG(LogTemp, Warning, TEXT("[FSessionsCommon::CheckJoinSessionState] Could not join session with user [%s] not logged in"), *ToLogString(LocalAccountId));
+			UE_LOG(LogTemp, Warning, TEXT("[FSessionsCommon::CheckJoinSessionState] Could not join session with user [%s] not logged in"), *ToLogString(Params.LocalAccountId));
 
-				return Errors::InvalidUser();
-			}
+			return Errors::InvalidUser();
 		}
 
 		// We check that the session is cached and valid for a join operation by the users
@@ -872,14 +828,11 @@ namespace UE::Online {
 
 		const FSessionSettings& SessionSettings = FoundSession->GetSessionSettings();
 
-		for (const FAccountId& LocalAccountId : LocalAccountIds)
+		if (SessionSettings.SessionMembers.Contains(Params.LocalAccountId))
 		{
-			if (SessionSettings.SessionMembers.Contains(LocalAccountId))
-			{
-				UE_LOG(LogTemp, Warning, TEXT("[FSessionsCommon::CheckJoinSessionState] Could not join session. User [%s] already in session"), *ToLogString(LocalAccountId));
+			UE_LOG(LogTemp, Warning, TEXT("[FSessionsCommon::CheckJoinSessionState] Could not join session. User [%s] already in session"), *ToLogString(Params.LocalAccountId));
 
-				return Errors::AccessDenied();
-			}
+			return Errors::AccessDenied();
 		}
 
 		if (!SessionSettings.bAllowNewMembers)
@@ -891,27 +844,24 @@ namespace UE::Online {
 
 		if (!SessionSettings.bAllowUnregisteredPlayers)
 		{
-			for (const FAccountId& LocalAccountId : LocalAccountIds)
+			if (!SessionSettings.RegisteredPlayers.Contains(Params.LocalAccountId))
 			{
-				if (!SessionSettings.RegisteredPlayers.Contains(LocalAccountId))
-				{
-					UE_LOG(LogTemp, Warning, TEXT("[FSessionsCommon::CheckJoinSessionState] Could not join session. bAllowUnregisteredPlayers is set to false and user [%s] is not registered"), *ToLogString(LocalAccountId));
+				UE_LOG(LogTemp, Warning, TEXT("[FSessionsCommon::CheckJoinSessionState] Could not join session. bAllowUnregisteredPlayers is set to false and user [%s] is not registered"), *ToLogString(Params.LocalAccountId));
 
-					return Errors::AccessDenied();
-				}
+				return Errors::AccessDenied();
 			}
 		}
 
-		if (SessionSettings.NumMaxPublicConnections > 0 && SessionSettings.NumOpenPublicConnections < (uint32)LocalAccountIds.Num())
+		if (SessionSettings.NumMaxPublicConnections > 0 && SessionSettings.NumOpenPublicConnections > 0)
 		{
-			UE_LOG(LogTemp, Warning, TEXT("[FSessionsCommon::CheckJoinSessionState] Could not join session. Not enough NumOpenPublicConnections [%u] for the amount of joining players [%d]"), SessionSettings.NumOpenPublicConnections, LocalAccountIds.Num());
+			UE_LOG(LogTemp, Warning, TEXT("[FSessionsCommon::CheckJoinSessionState] Could not join session. No open public connections"));
 
 			return Errors::AccessDenied();
 		}
 
-		if (SessionSettings.NumMaxPrivateConnections > 0 && SessionSettings.NumOpenPrivateConnections < (uint32)LocalAccountIds.Num())
+		if (SessionSettings.NumMaxPrivateConnections > 0 && SessionSettings.NumOpenPrivateConnections > 0)
 		{
-			UE_LOG(LogTemp, Warning, TEXT("[FSessionsCommon::CheckJoinSessionState] Could not join session. Not enough NumOpenPrivateConnections [%u] for the amount of joining players [%d]"), SessionSettings.NumOpenPrivateConnections, LocalAccountIds.Num());
+			UE_LOG(LogTemp, Warning, TEXT("[FSessionsCommon::CheckJoinSessionState] Could not join session. No open private connections"));
 
 			return Errors::AccessDenied();
 		}
@@ -919,11 +869,11 @@ namespace UE::Online {
 		return Errors::Success();
 	}
 
-	FOnlineError FSessionsCommon::CheckAddSessionMembersState(const FAddSessionMembers::Params& Params)
+	FOnlineError FSessionsCommon::CheckAddSessionMemberState(const FAddSessionMember::Params& Params)
 	{
 		if (TOptional<FOnlineError> Result = CheckSessionExistsByName(Params.LocalAccountId, Params.SessionName))
 		{
-			UE_LOG(LogTemp, Warning, TEXT("[FSessionsCommon::CheckAddSessionMembersState] Could not add session member to session with name [%s]. Session not found"), *Params.SessionName.ToString());
+			UE_LOG(LogTemp, Warning, TEXT("[FSessionsCommon::CheckAddSessionMemberState] Could not add session member to session with name [%s]. Session not found"), *Params.SessionName.ToString());
 
 			return Result.GetValue();
 		}
@@ -933,11 +883,11 @@ namespace UE::Online {
 		return Errors::Success();
 	}
 
-	FOnlineError FSessionsCommon::CheckRemoveSessionMembersState(const FRemoveSessionMembers::Params& Params)
+	FOnlineError FSessionsCommon::CheckRemoveSessionMemberState(const FRemoveSessionMember::Params& Params)
 	{
 		if (TOptional<FOnlineError> Result = CheckSessionExistsByName(Params.LocalAccountId, Params.SessionName))
 		{
-			UE_LOG(LogTemp, Warning, TEXT("[FSessionsCommon::CheckRemoveSessionMembersState] Could not remove session member from session with name [%s]. Session not found"), *Params.SessionName.ToString());
+			UE_LOG(LogTemp, Warning, TEXT("[FSessionsCommon::CheckRemoveSessionMemberState] Could not remove session member from session with name [%s]. Session not found"), *Params.SessionName.ToString());
 
 			return Result.GetValue();
 		}
@@ -955,17 +905,13 @@ namespace UE::Online {
 			return Result.GetValue();
 		}
 
+		// User login check
 		IAuthPtr Auth = Services.GetAuthInterface();
-
-		// User login check for all local users
-		for (const FAccountId& LocalAccountId : Params.LocalAccounts)
+		if (!Auth->IsLoggedIn(Params.LocalAccountId))
 		{
-			if (!Auth->IsLoggedIn(LocalAccountId))
-			{
-				UE_LOG(LogTemp, Warning, TEXT("[FSessionsCommon::CheckLeaveSessionState] Could not leave session with user [%s] not logged in"), *ToLogString(LocalAccountId));
+			UE_LOG(LogTemp, Warning, TEXT("[FSessionsCommon::CheckLeaveSessionState] Could not leave session with user [%s] not logged in"), *ToLogString(Params.LocalAccountId));
 
-				return Errors::InvalidUser();
-			}
+			return Errors::InvalidUser();
 		}
 
 		return Errors::Success();
diff --git a/Engine/Plugins/Online/Experimental/OnlineServices/Source/OnlineServicesCommon/Public/Online/SessionsCommon.h b/Engine/Plugins/Online/Experimental/OnlineServices/Source/OnlineServicesCommon/Public/Online/SessionsCommon.h
--- a/Engine/Plugins/Online/Experimental/OnlineServices/Source/OnlineServicesCommon/Public/Online/SessionsCommon.h
+++ b/Engine/Plugins/Online/Experimental/OnlineServices/Source/OnlineServicesCommon/Public/Online/SessionsCommon.h
@@ -210,8 +210,8 @@ public:
 	virtual TOnlineAsyncOpHandle<FFindSessions> FindSessions(FFindSessions::Params&& Params) override;
 	virtual TOnlineAsyncOpHandle<FStartMatchmaking> StartMatchmaking(FStartMatchmaking::Params&& Params) override;
 	virtual TOnlineAsyncOpHandle<FJoinSession> JoinSession(FJoinSession::Params&& Params) override;
-	virtual TOnlineAsyncOpHandle<FAddSessionMembers> AddSessionMembers(FAddSessionMembers::Params&& Params) override;
-	virtual TOnlineAsyncOpHandle<FRemoveSessionMembers> RemoveSessionMembers(FRemoveSessionMembers::Params&& Params) override;
+	virtual TOnlineAsyncOpHandle<FAddSessionMember> AddSessionMember(FAddSessionMember::Params&& Params) override;
+	virtual TOnlineAsyncOpHandle<FRemoveSessionMember> RemoveSessionMember(FRemoveSessionMember::Params&& Params) override;
 	virtual TOnlineAsyncOpHandle<FSendSessionInvite> SendSessionInvite(FSendSessionInvite::Params&& Params) override;
 	virtual TOnlineResult<FGetSessionInvites> GetSessionInvites(FGetSessionInvites::Params&& Params) override;
 	virtual TOnlineAsyncOpHandle<FRejectSessionInvite> RejectSessionInvite(FRejectSessionInvite::Params&& Params) override;
@@ -231,8 +231,8 @@ protected:
 	void ClearSessionByName(const FName& SessionName);
 	void ClearSessionById(const FOnlineSessionIdHandle& SessionId);
 
-	TOnlineResult<FAddSessionMembers> AddSessionMembersImpl(const FAddSessionMembers::Params& Params);
-	TOnlineResult<FRemoveSessionMembers> RemoveSessionMembersImpl(const FRemoveSessionMembers::Params& Params);
+	TOnlineResult<FAddSessionMember> AddSessionMemberImpl(const FAddSessionMember::Params& Params);
+	TOnlineResult<FRemoveSessionMember> RemoveSessionMemberImpl(const FRemoveSessionMember::Params& Params);
 
 	TOnlineResult<FRegisterPlayers> RegisterPlayersImpl(const FRegisterPlayers::Params& Params);
 	TOnlineResult<FUnregisterPlayers> UnregisterPlayersImpl(const FUnregisterPlayers::Params& Params);
@@ -251,9 +251,9 @@ protected:
 	FOnlineError CheckJoinSessionParams(const FJoinSession::Params& Params);
 	FOnlineError CheckJoinSessionState(const FJoinSession::Params& Params);
 
-	FOnlineError CheckAddSessionMembersState(const FAddSessionMembers::Params& Params);
+	FOnlineError CheckAddSessionMemberState(const FAddSessionMember::Params& Params);
 
-	FOnlineError CheckRemoveSessionMembersState(const FRemoveSessionMembers::Params& Params);
+	FOnlineError CheckRemoveSessionMemberState(const FRemoveSessionMember::Params& Params);
 
 	FOnlineError CheckLeaveSessionState(const FLeaveSession::Params& Params);
 
@@ -266,7 +266,7 @@ protected:
 	FOnlineError CheckUnregisterPlayersState(const FUnregisterPlayers::Params& Params);
 
 private:
-	TOptional<FOnlineError> CheckSessionExistsByName(const FAccountId& LocalUserId, const FName& SessionName);
+	TOptional<FOnlineError> CheckSessionExistsByName(const FAccountId& LocalAccountId, const FName& SessionName);
 
 protected:
 
diff --git a/Engine/Plugins/Online/Experimental/OnlineServices/Source/OnlineServicesCommonEngineUtils/Private/Online/SessionsLAN.cpp b/Engine/Plugins/Online/Experimental/OnlineServices/Source/OnlineServicesCommonEngineUtils/Private/Online/SessionsLAN.cpp
--- a/Engine/Plugins/Online/Experimental/OnlineServices/Source/OnlineServicesCommonEngineUtils/Private/Online/SessionsLAN.cpp
+++ b/Engine/Plugins/Online/Experimental/OnlineServices/Source/OnlineServicesCommonEngineUtils/Private/Online/SessionsLAN.cpp
@@ -156,8 +156,8 @@ TOnlineAsyncOpHandle<FCreateSession> FSessionsLAN::CreateSession(FCreateSession:
 		NewSessionLANRef->SessionId = FOnlineSessionIdRegistryLAN::GetChecked(Services.GetServicesProvider()).GetNextSessionId();
 		NewSessionLANRef->SessionSettings = OpParams.SessionSettings;
 
-		// For LAN sessions, we'll add all the Session members manually instead of calling JoinSession since there is no API calls involved
-		NewSessionLANRef->SessionSettings.SessionMembers.Append(OpParams.LocalAccounts);
+		// For LAN sessions, we'll add the Session member manually instead of calling JoinSession since there is no API calls involved
+		NewSessionLANRef->SessionSettings.SessionMembers.Emplace(OpParams.LocalAccountId, OpParams.SessionMemberData);
 
 		// We save the local object for the session, and set up the appropriate references
 		AllSessionsById.Emplace(NewSessionLANRef->SessionId, NewSessionLANRef);
@@ -316,11 +316,7 @@ TOnlineAsyncOpHandle<FJoinSession> FSessionsLAN::JoinSession(FJoinSession::Param
 
 		Op.SetResult(FJoinSession::Result{ });
 
-		TArray<FAccountId> LocalAccountIds;
-		LocalAccountIds.Reserve(OpParams.LocalAccounts.Num());
-		OpParams.LocalAccounts.GenerateKeyArray(LocalAccountIds);
-
-		FSessionJoined SessionJoinedEvent = { MoveTemp(LocalAccountIds), FoundSession->GetSessionId() };
+		FSessionJoined SessionJoinedEvent = { OpParams.LocalAccountId, FoundSession->GetSessionId() };
 		
 		SessionEvents.OnSessionJoined.Broadcast(SessionJoinedEvent);
 	})
@@ -368,7 +364,7 @@ TOnlineAsyncOpHandle<FLeaveSession> FSessionsLAN::LeaveSession(FLeaveSession::Pa
 		Op.SetResult(FLeaveSession::Result{ });
 
 		FSessionLeft SessionLeftEvent;
-		SessionLeftEvent.LocalAccountIds.Append(OpParams.LocalAccounts);
+		SessionLeftEvent.LocalAccountId = OpParams.LocalAccountId;
 		SessionEvents.OnSessionLeft.Broadcast(SessionLeftEvent);
 	})
 	.Enqueue(GetSerialQueue());
diff --git a/Engine/Plugins/Online/Experimental/OnlineServices/Source/OnlineServicesCommonEngineUtils/Public/Online/SessionsLAN.h b/Engine/Plugins/Online/Experimental/OnlineServices/Source/OnlineServicesCommonEngineUtils/Public/Online/SessionsLAN.h
--- a/Engine/Plugins/Online/Experimental/OnlineServices/Source/OnlineServicesCommonEngineUtils/Public/Online/SessionsLAN.h
+++ b/Engine/Plugins/Online/Experimental/OnlineServices/Source/OnlineServicesCommonEngineUtils/Public/Online/SessionsLAN.h
@@ -64,11 +64,11 @@ public:
 
 protected:
 	TOptional<FOnlineError> TryHostLANSession();
-	void FindLANSessions(const FAccountId& LocalUserId);
+	void FindLANSessions(const FAccountId& LocalAccountId);
 	void StopLANSession();
 	void OnValidQueryPacketReceived(uint8* PacketData, int32 PacketLength, uint64 ClientNonce);
-	void OnValidResponsePacketReceived(uint8* PacketData, int32 PacketLength, const FAccountId LocalUserId);
-	void OnLANSearchTimeout(const FAccountId LocalUserId);
+	void OnValidResponsePacketReceived(uint8* PacketData, int32 PacketLength, const FAccountId LocalAccountId);
+	void OnLANSearchTimeout(const FAccountId LocalAccountId);
 	virtual void AppendSessionToPacket(FNboSerializeToBuffer& Packet, const FSessionLAN& Session) = 0;
 	virtual void ReadSessionFromPacket(FNboSerializeFromBuffer& Packet, FSessionLAN& Session) = 0;
 
diff --git a/Engine/Plugins/Online/Experimental/OnlineServices/Source/OnlineServicesInterface/Public/Online/Sessions.h b/Engine/Plugins/Online/Experimental/OnlineServices/Source/OnlineServicesInterface/Public/Online/Sessions.h
--- a/Engine/Plugins/Online/Experimental/OnlineServices/Source/OnlineServicesInterface/Public/Online/Sessions.h
+++ b/Engine/Plugins/Online/Experimental/OnlineServices/Source/OnlineServicesInterface/Public/Online/Sessions.h
@@ -262,7 +262,7 @@ struct FGetPresenceSession
 
 	struct Params
 	{
-		FAccountId LocalUserId;
+		FAccountId LocalAccountId;
 	};
 
 	struct Result
@@ -277,7 +277,7 @@ struct FSetPresenceSession
 
 	struct Params
 	{
-		FAccountId LocalUserId;
+		FAccountId LocalAccountId;
 
 		FOnlineSessionIdHandle SessionId;
 	};
@@ -293,7 +293,7 @@ struct FClearPresenceSession
 
 	struct Params
 	{
-		FAccountId LocalUserId;
+		FAccountId LocalAccountId;
 	};
 
 	struct Result
@@ -313,11 +313,11 @@ struct FCreateSession
 		/** The local name for the session */
 		FName SessionName;
 
+		/** Information for the local user who will join the session after creation */
+		FSessionMember SessionMemberData;
+
 		/** Settings object to define session properties during creation */
 		FSessionSettings SessionSettings;
-
-		/** Information for all local users who will join the session (includes the session creator) */
-		FSessionMembersMap LocalAccounts;
 	};
 
 	struct Result
@@ -354,15 +354,12 @@ struct FLeaveSession
 
 	struct Params
 	{
-		/* The local user agent which leaves the session*/
+		/* The local user agent which leaves the session */
 		FAccountId LocalAccountId;
 
 		/* The local name for the session. */
 		FName SessionName;
 
-		/** Ids for all local users who will leave the session (includes the main caller) */
-		TArray<FAccountId> LocalAccounts;
-
 		/* Whether the call should attempt to destroy the session instead of just leave it */
 		bool bDestroySession;
 	};
@@ -413,12 +410,12 @@ struct FStartMatchmaking
 		/** The local user agent which will perform the action. */
 		FAccountId LocalAccountId;
 
-		/* Information for all local users who will join the session (includes the session creator) */
-		FSessionMembersMap LocalAccounts;
-
 		/* Local name for the session */
 		FName SessionName;
 
+		/* Information for the local user who will create or join the session */
+		FSessionMember SessionMemberData;
+
 		/* Preferred settings to be used during session creation */
 		FSessionSettings SessionSettings;
 
@@ -447,8 +444,8 @@ struct FJoinSession
 		/* Id handle for the session to be joined. To be retrieved via session search or invite */
 		FOnlineSessionIdHandle SessionId;
 
-		/* Information for all local users who will join the session (includes the session creator). Any player that wants to join the session needs defined information to become a new member */
-		FSessionMembersMap LocalAccounts;
+		/* Information for the local user who will join the session */
+		FSessionMember SessionMemberData;
 	};
 
 	struct Result
@@ -457,9 +454,9 @@ struct FJoinSession
 	};
 };
 
-struct FAddSessionMembers
+struct FAddSessionMember
 {
-	static constexpr TCHAR Name[] = TEXT("AddSessionMembers");
+	static constexpr TCHAR Name[] = TEXT("AddSessionMember");
 
 	struct Params
 	{
@@ -469,8 +466,8 @@ struct FAddSessionMembers
 		/* Local name for the session */
 		FName SessionName;
 
-		/** Information for the session members to be added to the session. Any player that joins the session becomes a new member in doing so */
-		FSessionMembersMap NewSessionMembers;
+		/** Information for the session member to be added to the session. Any player that joins the session becomes a new member in doing so */
+		FSessionMember NewSessionMember;
 
 		/** Whether or not the new session members should also be added to the list of registered players. True by default*/
 		bool bRegisterPlayers = true;
@@ -482,9 +479,9 @@ struct FAddSessionMembers
 	};
 };
 
-struct FRemoveSessionMembers
+struct FRemoveSessionMember
 {
-	static constexpr TCHAR Name[] = TEXT("RemoveSessionMembers");
+	static constexpr TCHAR Name[] = TEXT("RemoveSessionMember");
 
 	struct Params
 	{
@@ -494,9 +491,6 @@ struct FRemoveSessionMembers
 		/* Local name for the session */
 		FName SessionName;
 
-		/* Id handles for the session members to be removed from the session */
-		TArray<FAccountId> SessionMemberIds;
-
 		/** Whether or not the session members should also be removed from the list of registered players. True by default*/
 		bool bUnregisterPlayers = true;
 	};
@@ -616,17 +610,17 @@ struct FUnregisterPlayers
 
 struct FSessionJoined
 {
-	/* The local users which joined the session */
-	TArray<FAccountId> LocalAccountIds;
+	/* The local user which joined the session */
+	FAccountId LocalAccountId;
 
-	/* A shared reference to the session joined. */
+	/* Id for the session joined. */
 	FOnlineSessionIdHandle SessionId;
 };
 
 struct FSessionLeft
 {
 	/* The local users which left the session */
-	TArray<FAccountId> LocalAccountIds;
+	FAccountId LocalAccountId;
 };
 
 struct FSessionUpdated
@@ -782,7 +776,7 @@ public:
 	 * @params Parameters for the AddSessionMember call
 	 * @return
 	 */
-	virtual TOnlineAsyncOpHandle<FAddSessionMembers> AddSessionMembers(FAddSessionMembers::Params&& Params) = 0;
+	virtual TOnlineAsyncOpHandle<FAddSessionMember> AddSessionMember(FAddSessionMember::Params&& Params) = 0;
 
 	/**
 	 * Removes a set of session member from the named session
@@ -793,7 +787,7 @@ public:
 	 * @params Parameters for the RemoveSessionMember call
 	 * @return
 	 */
-	virtual TOnlineAsyncOpHandle<FRemoveSessionMembers> RemoveSessionMembers(FRemoveSessionMembers::Params&& Params) = 0;
+	virtual TOnlineAsyncOpHandle<FRemoveSessionMember> RemoveSessionMember(FRemoveSessionMember::Params&& Params) = 0;
 
 	/**
 	 * Sends an invite to the named session to all given users.
@@ -982,7 +976,7 @@ BEGIN_ONLINE_STRUCT_META(FGetSessionById::Result)
 END_ONLINE_STRUCT_META()
 
 BEGIN_ONLINE_STRUCT_META(FGetPresenceSession::Params)
-	ONLINE_STRUCT_FIELD(FGetPresenceSession::Params, LocalUserId)
+	ONLINE_STRUCT_FIELD(FGetPresenceSession::Params, LocalAccountId)
 END_ONLINE_STRUCT_META()
 
 BEGIN_ONLINE_STRUCT_META(FGetPresenceSession::Result)
@@ -990,7 +984,7 @@ BEGIN_ONLINE_STRUCT_META(FGetPresenceSession::Result)
 END_ONLINE_STRUCT_META()
 
 BEGIN_ONLINE_STRUCT_META(FSetPresenceSession::Params)
-	ONLINE_STRUCT_FIELD(FSetPresenceSession::Params, LocalUserId),
+	ONLINE_STRUCT_FIELD(FSetPresenceSession::Params, LocalAccountId),
 	ONLINE_STRUCT_FIELD(FSetPresenceSession::Params, SessionId)
 END_ONLINE_STRUCT_META()
 
@@ -998,7 +992,7 @@ BEGIN_ONLINE_STRUCT_META(FSetPresenceSession::Result)
 END_ONLINE_STRUCT_META()
 
 BEGIN_ONLINE_STRUCT_META(FClearPresenceSession::Params)
-	ONLINE_STRUCT_FIELD(FClearPresenceSession::Params, LocalUserId)
+	ONLINE_STRUCT_FIELD(FClearPresenceSession::Params, LocalAccountId)
 END_ONLINE_STRUCT_META()
 
 BEGIN_ONLINE_STRUCT_META(FClearPresenceSession::Result)
@@ -1007,8 +1001,8 @@ END_ONLINE_STRUCT_META()
 BEGIN_ONLINE_STRUCT_META(FCreateSession::Params)
 	ONLINE_STRUCT_FIELD(FCreateSession::Params, LocalAccountId),
 	ONLINE_STRUCT_FIELD(FCreateSession::Params, SessionName),
-	ONLINE_STRUCT_FIELD(FCreateSession::Params, SessionSettings),
-	ONLINE_STRUCT_FIELD(FCreateSession::Params, LocalAccounts)
+	ONLINE_STRUCT_FIELD(FCreateSession::Params, SessionMemberData),
+	ONLINE_STRUCT_FIELD(FCreateSession::Params, SessionSettings)
 END_ONLINE_STRUCT_META()
 
 BEGIN_ONLINE_STRUCT_META(FCreateSession::Result)
@@ -1026,7 +1020,6 @@ END_ONLINE_STRUCT_META()
 BEGIN_ONLINE_STRUCT_META(FLeaveSession::Params)
 	ONLINE_STRUCT_FIELD(FLeaveSession::Params, LocalAccountId),
 	ONLINE_STRUCT_FIELD(FLeaveSession::Params, SessionName),
-	ONLINE_STRUCT_FIELD(FLeaveSession::Params, LocalAccounts),
 	ONLINE_STRUCT_FIELD(FLeaveSession::Params, bDestroySession)
 END_ONLINE_STRUCT_META()
 
@@ -1048,8 +1041,8 @@ END_ONLINE_STRUCT_META()
 
 BEGIN_ONLINE_STRUCT_META(FStartMatchmaking::Params)
 	ONLINE_STRUCT_FIELD(FStartMatchmaking::Params, LocalAccountId),
-	ONLINE_STRUCT_FIELD(FStartMatchmaking::Params, LocalAccounts),
 	ONLINE_STRUCT_FIELD(FStartMatchmaking::Params, SessionName),
+	ONLINE_STRUCT_FIELD(FStartMatchmaking::Params, SessionMemberData),
 	ONLINE_STRUCT_FIELD(FStartMatchmaking::Params, SessionSettings),
 	ONLINE_STRUCT_FIELD(FStartMatchmaking::Params, SearchFilters)
 END_ONLINE_STRUCT_META()
@@ -1061,30 +1054,29 @@ BEGIN_ONLINE_STRUCT_META(FJoinSession::Params)
 	ONLINE_STRUCT_FIELD(FJoinSession::Params, LocalAccountId),
 	ONLINE_STRUCT_FIELD(FJoinSession::Params, SessionName),
 	ONLINE_STRUCT_FIELD(FJoinSession::Params, SessionId),
-	ONLINE_STRUCT_FIELD(FJoinSession::Params, LocalAccounts)
+	ONLINE_STRUCT_FIELD(FJoinSession::Params, SessionMemberData)
 END_ONLINE_STRUCT_META()
 
 BEGIN_ONLINE_STRUCT_META(FJoinSession::Result)
 END_ONLINE_STRUCT_META()
 
-BEGIN_ONLINE_STRUCT_META(FAddSessionMembers::Params)
-	ONLINE_STRUCT_FIELD(FAddSessionMembers::Params, LocalAccountId),
-	ONLINE_STRUCT_FIELD(FAddSessionMembers::Params, SessionName),
-	ONLINE_STRUCT_FIELD(FAddSessionMembers::Params, NewSessionMembers),
-	ONLINE_STRUCT_FIELD(FAddSessionMembers::Params, bRegisterPlayers)
+BEGIN_ONLINE_STRUCT_META(FAddSessionMember::Params)
+	ONLINE_STRUCT_FIELD(FAddSessionMember::Params, LocalAccountId),
+	ONLINE_STRUCT_FIELD(FAddSessionMember::Params, SessionName),
+	ONLINE_STRUCT_FIELD(FAddSessionMember::Params, NewSessionMember),
+	ONLINE_STRUCT_FIELD(FAddSessionMember::Params, bRegisterPlayers)
 END_ONLINE_STRUCT_META()
 
-BEGIN_ONLINE_STRUCT_META(FAddSessionMembers::Result)
+BEGIN_ONLINE_STRUCT_META(FAddSessionMember::Result)
 END_ONLINE_STRUCT_META()
 
-BEGIN_ONLINE_STRUCT_META(FRemoveSessionMembers::Params)
-	ONLINE_STRUCT_FIELD(FRemoveSessionMembers::Params, LocalAccountId),
-	ONLINE_STRUCT_FIELD(FRemoveSessionMembers::Params, SessionName),
-	ONLINE_STRUCT_FIELD(FRemoveSessionMembers::Params, SessionMemberIds),
-	ONLINE_STRUCT_FIELD(FRemoveSessionMembers::Params, bUnregisterPlayers)
+BEGIN_ONLINE_STRUCT_META(FRemoveSessionMember::Params)
+	ONLINE_STRUCT_FIELD(FRemoveSessionMember::Params, LocalAccountId),
+	ONLINE_STRUCT_FIELD(FRemoveSessionMember::Params, SessionName),
+	ONLINE_STRUCT_FIELD(FRemoveSessionMember::Params, bUnregisterPlayers)
 END_ONLINE_STRUCT_META()
 
-BEGIN_ONLINE_STRUCT_META(FRemoveSessionMembers::Result)
+BEGIN_ONLINE_STRUCT_META(FRemoveSessionMember::Result)
 END_ONLINE_STRUCT_META()
 
 BEGIN_ONLINE_STRUCT_META(FSendSessionInvite::Params)
diff --git a/Engine/Plugins/Online/Experimental/OnlineServicesEOSGS/Source/Private/Online/SessionsEOSGS.cpp b/Engine/Plugins/Online/Experimental/OnlineServicesEOSGS/Source/Private/Online/SessionsEOSGS.cpp
--- a/Engine/Plugins/Online/Experimental/OnlineServicesEOSGS/Source/Private/Online/SessionsEOSGS.cpp
+++ b/Engine/Plugins/Online/Experimental/OnlineServicesEOSGS/Source/Private/Online/SessionsEOSGS.cpp
@@ -1162,8 +1162,7 @@ TOnlineAsyncOpHandle<FLeaveSession> FSessionsEOSGS::LeaveSession(FLeaveSession::
 		Op.SetResult(FLeaveSession::Result{ });
 
 		FSessionLeft SessionLeftEvent;
-		SessionLeftEvent.LocalAccountIds.Add(OpParams.LocalAccountId);
-		//SessionLeftEvent.LocalAccountIds.Append(OpParams.AdditionalLocalUsers); // TODO: Process multiple users in Session destruction
+		SessionLeftEvent.LocalAccountId = OpParams.LocalAccountId;
 		SessionEvents.OnSessionLeft.Broadcast(SessionLeftEvent);
 	})
 	.Enqueue(GetSerialQueue());
@@ -1842,24 +1841,24 @@ TOnlineAsyncOpHandle<FRejectSessionInvite> FSessionsEOSGS::RejectSessionInvite(F
 	return Op->GetHandle();
 }
 
-TOnlineAsyncOpHandle<FAddSessionMembers> FSessionsEOSGS::AddSessionMembers(FAddSessionMembers::Params&& Params)
+TOnlineAsyncOpHandle<FAddSessionMember> FSessionsEOSGS::AddSessionMember(FAddSessionMember::Params&& Params)
 {
 	// LAN Sessions
 	TOnlineResult<FGetSessionByName> Result = GetSessionByName({ Params.SessionName });
 	if (Result.IsOk() && Result.GetOkValue().Session->GetSessionSettings().bIsLANSession)
 	{
-		return FSessionsLAN::AddSessionMembers(MoveTemp(Params));
+		return FSessionsLAN::AddSessionMember(MoveTemp(Params));
 	}
 
 	// EOSGS Sessions
 
-	TOnlineAsyncOpRef<FAddSessionMembers> Op = GetOp<FAddSessionMembers>(MoveTemp(Params));
+	TOnlineAsyncOpRef<FAddSessionMember> Op = GetOp<FAddSessionMember>(MoveTemp(Params));
 
-	Op->Then([this](TOnlineAsyncOp<FAddSessionMembers>& Op, TPromise<const EOS_Sessions_RegisterPlayersCallbackInfo*>&& Promise)
+	Op->Then([this](TOnlineAsyncOp<FAddSessionMember>& Op, TPromise<const EOS_Sessions_RegisterPlayersCallbackInfo*>&& Promise)
 	{
-		const FAddSessionMembers::Params& OpParams = Op.GetParams();
+		const FAddSessionMember::Params& OpParams = Op.GetParams();
 
-		FOnlineError StateCheck = CheckAddSessionMembersState(OpParams);
+		FOnlineError StateCheck = CheckAddSessionMemberState(OpParams);
 		if (StateCheck != Errors::Success())
 		{
 			Op.SetError(MoveTemp(StateCheck));
@@ -1871,17 +1870,10 @@ TOnlineAsyncOpHandle<FAddSessionMembers> FSessionsEOSGS::AddSessionMembers(FAddS
 		Options.ApiVersion = EOS_SESSIONS_REGISTERPLAYERS_API_LATEST;
 		static_assert(EOS_SESSIONS_REGISTERPLAYERS_API_LATEST == 2, "EOS_Sessions_RegisterPlayersOptions updated, check new fields");
 
-		Options.PlayersToRegisterCount = OpParams.NewSessionMembers.Num();
-
-		TArray<EOS_ProductUserId> PlayersToRegister;
-		TArray<FAccountId> SessionMemberIds;
-		OpParams.NewSessionMembers.GenerateKeyArray(SessionMemberIds);		
-		for (FAccountId IdHandle : SessionMemberIds)
-		{
-			// TODO: Do this with ResolveAccountIds instead
-			PlayersToRegister.Add(GetProductUserIdChecked(IdHandle));
-		}
+		Options.PlayersToRegisterCount = 1;
 
+		// TODO: Do this with ResolveAccountIds instead
+		TArray<EOS_ProductUserId> PlayersToRegister = { GetProductUserIdChecked(OpParams.LocalAccountId) };
 		Options.PlayersToRegister = PlayersToRegister.GetData();
 
 		const FTCHARToUTF8 SessionNameUtf8(*OpParams.SessionName.ToString());
@@ -1889,7 +1881,7 @@ TOnlineAsyncOpHandle<FAddSessionMembers> FSessionsEOSGS::AddSessionMembers(FAddS
 
 		EOS_Async(EOS_Sessions_RegisterPlayers, SessionsHandle, Options, MoveTemp(Promise));
 	})
-	.Then([this](TOnlineAsyncOp<FAddSessionMembers>& Op, const EOS_Sessions_RegisterPlayersCallbackInfo* Result)
+	.Then([this](TOnlineAsyncOp<FAddSessionMember>& Op, const EOS_Sessions_RegisterPlayersCallbackInfo* Result)
 	{
 		if (Result->ResultCode != EOS_EResult::EOS_Success)
 		{
@@ -1897,7 +1889,7 @@ TOnlineAsyncOpHandle<FAddSessionMembers> FSessionsEOSGS::AddSessionMembers(FAddS
 			return;
 		}
 
-		TOnlineResult<FAddSessionMembers> ImplResult = AddSessionMembersImpl(Op.GetParams());
+		TOnlineResult<FAddSessionMember> ImplResult = AddSessionMemberImpl(Op.GetParams());
 		if (ImplResult.IsOk())
 		{
 			Op.SetResult(MoveTemp(ImplResult.GetOkValue()));
@@ -1912,24 +1904,24 @@ TOnlineAsyncOpHandle<FAddSessionMembers> FSessionsEOSGS::AddSessionMembers(FAddS
 	return Op->GetHandle();
 }
 
-TOnlineAsyncOpHandle<FRemoveSessionMembers> FSessionsEOSGS::RemoveSessionMembers(FRemoveSessionMembers::Params&& Params)
+TOnlineAsyncOpHandle<FRemoveSessionMember> FSessionsEOSGS::RemoveSessionMember(FRemoveSessionMember::Params&& Params)
 {
 	// LAN Sessions
 	TOnlineResult<FGetSessionByName> Result = GetSessionByName({ Params.SessionName });
 	if (Result.IsOk() && Result.GetOkValue().Session->GetSessionSettings().bIsLANSession)
 	{
-		return FSessionsLAN::RemoveSessionMembers(MoveTemp(Params));
+		return FSessionsLAN::RemoveSessionMember(MoveTemp(Params));
 	}
 
 	// EOSGS Sessions
 
-	TOnlineAsyncOpRef<FRemoveSessionMembers> Op = GetOp<FRemoveSessionMembers>(MoveTemp(Params));
+	TOnlineAsyncOpRef<FRemoveSessionMember> Op = GetOp<FRemoveSessionMember>(MoveTemp(Params));
 
-	Op->Then([this](TOnlineAsyncOp<FRemoveSessionMembers>& Op, TPromise<const EOS_Sessions_UnregisterPlayersCallbackInfo*>&& Promise)
+	Op->Then([this](TOnlineAsyncOp<FRemoveSessionMember>& Op, TPromise<const EOS_Sessions_UnregisterPlayersCallbackInfo*>&& Promise)
 	{
-		const FRemoveSessionMembers::Params& OpParams = Op.GetParams();
+		const FRemoveSessionMember::Params& OpParams = Op.GetParams();
 
-		FOnlineError StateCheck = CheckRemoveSessionMembersState(OpParams);
+		FOnlineError StateCheck = CheckRemoveSessionMemberState(OpParams);
 		if (StateCheck != Errors::Success())
 		{
 			Op.SetError(MoveTemp(StateCheck));
@@ -1941,15 +1933,10 @@ TOnlineAsyncOpHandle<FRemoveSessionMembers> FSessionsEOSGS::RemoveSessionMembers
 		Options.ApiVersion = EOS_SESSIONS_UNREGISTERPLAYERS_API_LATEST;
 		static_assert(EOS_SESSIONS_UNREGISTERPLAYERS_API_LATEST == 2, "EOS_Sessions_UnregisterPlayersOptions updated, check new fields");
 
-		Options.PlayersToUnregisterCount = OpParams.SessionMemberIds.Num();
-
-		TArray<EOS_ProductUserId> PlayersToUnregister;
-		for (FAccountId IdHandle : OpParams.SessionMemberIds)
-		{
-			// TODO: Do this with ResolveAccountIds instead
-			PlayersToUnregister.Add(GetProductUserIdChecked(IdHandle));
-		}
+		Options.PlayersToUnregisterCount = 1;
 
+		// TODO: Do this with ResolveAccountIds instead
+		TArray<EOS_ProductUserId> PlayersToUnregister{ GetProductUserIdChecked(OpParams.LocalAccountId) };
 		Options.PlayersToUnregister = PlayersToUnregister.GetData();
 
 		const FTCHARToUTF8 SessionNameUtf8(*OpParams.SessionName.ToString());
@@ -1957,7 +1944,7 @@ TOnlineAsyncOpHandle<FRemoveSessionMembers> FSessionsEOSGS::RemoveSessionMembers
 
 		EOS_Async(EOS_Sessions_UnregisterPlayers, SessionsHandle, Options, MoveTemp(Promise));
 	})
-	.Then([this](TOnlineAsyncOp<FRemoveSessionMembers>& Op, const EOS_Sessions_UnregisterPlayersCallbackInfo* Result)
+	.Then([this](TOnlineAsyncOp<FRemoveSessionMember>& Op, const EOS_Sessions_UnregisterPlayersCallbackInfo* Result)
 	{
 		if (Result->ResultCode != EOS_EResult::EOS_Success)
 		{
@@ -1965,7 +1952,7 @@ TOnlineAsyncOpHandle<FRemoveSessionMembers> FSessionsEOSGS::RemoveSessionMembers
 			return;
 		}
 
-		TOnlineResult<FRemoveSessionMembers> ImplResult = RemoveSessionMembersImpl(Op.GetParams());
+		TOnlineResult<FRemoveSessionMember> ImplResult = RemoveSessionMemberImpl(Op.GetParams());
 		if (ImplResult.IsOk())
 		{
 			Op.SetResult(MoveTemp(ImplResult.GetOkValue()));
diff --git a/Engine/Plugins/Online/Experimental/OnlineServicesEOSGS/Source/Public/Online/SessionsEOSGS.h b/Engine/Plugins/Online/Experimental/OnlineServicesEOSGS/Source/Public/Online/SessionsEOSGS.h
--- a/Engine/Plugins/Online/Experimental/OnlineServicesEOSGS/Source/Public/Online/SessionsEOSGS.h
+++ b/Engine/Plugins/Online/Experimental/OnlineServicesEOSGS/Source/Public/Online/SessionsEOSGS.h
@@ -197,8 +197,8 @@ public:
 	virtual TOnlineAsyncOpHandle<FJoinSession> JoinSession(FJoinSession::Params&& Params) override;
 	virtual TOnlineAsyncOpHandle<FSendSessionInvite> SendSessionInvite(FSendSessionInvite::Params&& Params) override;
 	virtual TOnlineAsyncOpHandle<FRejectSessionInvite> RejectSessionInvite(FRejectSessionInvite::Params&& Params) override;
-	virtual TOnlineAsyncOpHandle<FAddSessionMembers> AddSessionMembers(FAddSessionMembers::Params&& Params) override;
-	virtual TOnlineAsyncOpHandle<FRemoveSessionMembers> RemoveSessionMembers(FRemoveSessionMembers::Params&& Params) override;
+	virtual TOnlineAsyncOpHandle<FAddSessionMember> AddSessionMember(FAddSessionMember::Params&& Params) override;
+	virtual TOnlineAsyncOpHandle<FRemoveSessionMember> RemoveSessionMember(FRemoveSessionMember::Params&& Params) override;
 
 protected:
 	void RegisterEventHandlers();
@@ -261,12 +261,12 @@ protected:
 	/**
 	 * Builds a session from an invite id, calling BuildSessionFromDetailsHandle
 	 */
-	TOnlineAsyncOpHandle<FBuildSessionFromDetailsHandle> BuildSessionFromInvite(const FAccountId& LocalUserId, const FString& InInviteId);
+	TOnlineAsyncOpHandle<FBuildSessionFromDetailsHandle> BuildSessionFromInvite(const FAccountId& LocalAccountId, const FString& InInviteId);
 
 	/**
 	 * Builds a session from a UI event id, calling BuildSessionFromDetailsHandle
 	 */
-	TOnlineAsyncOpHandle<FBuildSessionFromDetailsHandle> BuildSessionFromUIEvent(const FAccountId& LocalUserId, const EOS_UI_EventId& UIEventId);
+	TOnlineAsyncOpHandle<FBuildSessionFromDetailsHandle> BuildSessionFromUIEvent(const FAccountId& LocalAccountId, const EOS_UI_EventId& UIEventId);
 
 private:
 	// FSessionsLAN
diff --git a/Engine/Plugins/Online/Experimental/OnlineServicesOSSAdapter/Source/Private/Online/SessionsOSSAdapter.cpp b/Engine/Plugins/Online/Experimental/OnlineServicesOSSAdapter/Source/Private/Online/SessionsOSSAdapter.cpp
--- a/Engine/Plugins/Online/Experimental/OnlineServicesOSSAdapter/Source/Private/Online/SessionsOSSAdapter.cpp
+++ b/Engine/Plugins/Online/Experimental/OnlineServicesOSSAdapter/Source/Private/Online/SessionsOSSAdapter.cpp
@@ -299,15 +299,15 @@ void FSessionsOSSAdapter::Initialize()
 	{
 		FOnlineServicesOSSAdapter& ServicesOSSAdapter = static_cast<FOnlineServicesOSSAdapter&>(Services);
 
-		FUniqueNetIdRef LocalUserId = FUniqueNetIdString::EmptyId();
-		FUniqueNetIdPtr LocalUserIdPtr = IdentityInterface->GetUniquePlayerId(ControllerId);
-		if (LocalUserIdPtr.IsValid())
+		FUniqueNetIdRef LocalAccountId = FUniqueNetIdString::EmptyId();
+		FUniqueNetIdPtr LocalAccountIdPtr = IdentityInterface->GetUniquePlayerId(ControllerId);
+		if (LocalAccountIdPtr.IsValid())
 		{
-			LocalUserId = (*LocalUserIdPtr).AsShared();
+			LocalAccountId = (*LocalAccountIdPtr).AsShared();
 		}
 
 		FUISessionJoinRequested Event {
-			ServicesOSSAdapter.GetAccountIdRegistry().FindOrAddHandle(LocalUserId),
+			ServicesOSSAdapter.GetAccountIdRegistry().FindOrAddHandle(LocalAccountId),
 			TResult<FOnlineSessionIdHandle, FOnlineError>(GetSessionIdRegistry().FindOrAddHandle(InviteResult.Session.SessionInfo->GetSessionId().AsShared())),
 			EUISessionJoinRequestedSource::FromInvitation
 		};
@@ -438,7 +438,7 @@ TOnlineAsyncOpHandle<FCreateSession> FSessionsOSSAdapter::CreateSession(FCreateS
 		}
 
 		// Dedicated servers can update sessions so we won't check if LocalAccountId is valid
-		const FUniqueNetIdRef LocalUserId = Services.Get<FAuthOSSAdapter>()->GetUniqueNetId(OpParams.LocalAccountId).ToSharedRef();
+		const FUniqueNetIdRef LocalAccountId = Services.Get<FAuthOSSAdapter>()->GetUniqueNetId(OpParams.LocalAccountId).ToSharedRef();
 
 		MakeMulticastAdapter(this, SessionsInterface->OnCreateSessionCompleteDelegates,
 		[this, WeakOp = Op.AsWeak()](FName SessionName, const bool bWasSuccessful) mutable
@@ -472,7 +472,7 @@ TOnlineAsyncOpHandle<FCreateSession> FSessionsOSSAdapter::CreateSession(FCreateS
 			}
 		});
 
-		SessionsInterface->CreateSession(*LocalUserId, OpParams.SessionName, BuildV1Settings(OpParams.SessionSettings));
+		SessionsInterface->CreateSession(*LocalAccountId, OpParams.SessionName, BuildV1Settings(OpParams.SessionSettings));
 
 		// TODO: Pending support for multiple local users
 	})
@@ -598,7 +598,7 @@ TOnlineAsyncOpHandle<FLeaveSession> FSessionsOSSAdapter::LeaveSession(FLeaveSess
 				Op->SetResult({ });
 
 				FSessionLeft SessionLeftEvent;
-				SessionLeftEvent.LocalAccountIds.Append(Op->GetParams().LocalAccounts);
+				SessionLeftEvent.LocalAccountId = Op->GetParams().LocalAccountId;
 
 				SessionEvents.OnSessionLeft.Broadcast(SessionLeftEvent);
 			}
@@ -638,8 +638,8 @@ TOnlineAsyncOpHandle<FFindSessions> FSessionsOSSAdapter::FindSessions(FFindSessi
 			return;
 		}
 
-		const FUniqueNetIdRef LocalUserId = Services.Get<FAuthOSSAdapter>()->GetUniqueNetId(OpParams.LocalAccountId).ToSharedRef();
-		if (!LocalUserId->IsValid())
+		const FUniqueNetIdRef LocalAccountId = Services.Get<FAuthOSSAdapter>()->GetUniqueNetId(OpParams.LocalAccountId).ToSharedRef();
+		if (!LocalAccountId->IsValid())
 		{
 			Op.SetError(Errors::InvalidUser());
 			return;
@@ -697,7 +697,7 @@ TOnlineAsyncOpHandle<FFindSessions> FSessionsOSSAdapter::FindSessions(FFindSessi
 
 				PendingV1SessionSearchesPerUser.Emplace(OpParams.LocalAccountId, MakeShared<FOnlineSessionSearch>());
 
-				SessionsInterface->FindSessionById(*LocalUserId, *SessionIdRef, *TargetIdRef, *MakeDelegateAdapter(this, [this, WeakOp = Op.AsWeak()](int32 LocalUserNum, bool bWasSuccessful, const FOnlineSessionSearchResult& SingleSearchResult) mutable
+				SessionsInterface->FindSessionById(*LocalAccountId, *SessionIdRef, *TargetIdRef, *MakeDelegateAdapter(this, [this, WeakOp = Op.AsWeak()](int32 LocalUserNum, bool bWasSuccessful, const FOnlineSessionSearchResult& SingleSearchResult) mutable
 				{
 					if (TOnlineAsyncOpPtr<FFindSessions> Op = WeakOp.Pin())
 					{
@@ -724,7 +724,7 @@ TOnlineAsyncOpHandle<FFindSessions> FSessionsOSSAdapter::FindSessions(FFindSessi
 			int32 LocalUserNum = Services.Get<FAuthOSSAdapter>()->GetLocalUserNum(OpParams.LocalAccountId);
 
 			MakeMulticastAdapter(this, SessionsInterface->OnFindFriendSessionCompleteDelegates[LocalUserNum],
-			[this, WeakOp = Op.AsWeak()](int32 LocalUserId, bool bWasSuccessful, const TArray<FOnlineSessionSearchResult>& FriendSearchResults)
+			[this, WeakOp = Op.AsWeak()](int32 LocalAccountId, bool bWasSuccessful, const TArray<FOnlineSessionSearchResult>& FriendSearchResults)
 			{
 				if (TOnlineAsyncOpPtr<FFindSessions> Op = WeakOp.Pin())
 				{
@@ -756,7 +756,7 @@ TOnlineAsyncOpHandle<FFindSessions> FSessionsOSSAdapter::FindSessions(FFindSessi
 			// The cached search information is not used in this mode, but we'll still update it to not have more than one search at the same time
 			PendingV1SessionSearchesPerUser.Emplace(OpParams.LocalAccountId, MakeShared<FOnlineSessionSearch>());
 
-			SessionsInterface->FindFriendSession(*LocalUserId, TargetUniqueNetIds);
+			SessionsInterface->FindFriendSession(*LocalAccountId, TargetUniqueNetIds);
 		}
 		else
 		{
@@ -790,7 +790,7 @@ TOnlineAsyncOpHandle<FFindSessions> FSessionsOSSAdapter::FindSessions(FFindSessi
 
 			TSharedRef<FOnlineSessionSearch>& PendingV1SessionSearch = PendingV1SessionSearchesPerUser.Emplace(OpParams.LocalAccountId, BuildV1SessionSearch(OpParams));
 
-			SessionsInterface->FindSessions(*LocalUserId, PendingV1SessionSearch);
+			SessionsInterface->FindSessions(*LocalAccountId, PendingV1SessionSearch);
 		}
 	})
 	.Enqueue(GetSerialQueue());
@@ -830,32 +830,30 @@ TOnlineAsyncOpHandle<FStartMatchmaking> FSessionsOSSAdapter::StartMatchmaking(FS
 		FOnlineServicesOSSAdapter& ServicesOSSAdapter = static_cast<FOnlineServicesOSSAdapter&>(Services);
 
 		TArray<FSessionMatchmakingUser> MatchMakingUsers;
-		for (const TPair<FAccountId, FSessionMember>& LocalUser : OpParams.LocalAccounts)
-		{
-			FSessionMatchmakingUser NewMatchMakingUser { ServicesOSSAdapter.GetAccountIdRegistry().GetIdValue(LocalUser.Key).ToSharedRef() };
 
-			for (const TPair<FSchemaAttributeId, FCustomSessionSetting>& Entry : LocalUser.Value.MemberSettings)
+		FSessionMatchmakingUser NewMatchMakingUser { ServicesOSSAdapter.GetAccountIdRegistry().GetIdValue(OpParams.LocalAccountId).ToSharedRef() };
+
+		for (const TPair<FSchemaAttributeId, FCustomSessionSetting>& Entry : OpParams.SessionMemberData.MemberSettings)
+		{
+			switch (Entry.Value.Data.VariantType)
 			{
-				switch (Entry.Value.Data.VariantType)
-				{
-				case ESchemaAttributeType::Bool:
-					NewMatchMakingUser.Attributes.Emplace(Entry.Key.ToString(), FVariantData(Entry.Value.Data.GetBoolean()));
-					break;
-				case ESchemaAttributeType::Double:
-					NewMatchMakingUser.Attributes.Emplace(Entry.Key.ToString(), FVariantData(Entry.Value.Data.GetDouble()));
-					break;
-				case ESchemaAttributeType::Int64:
-					NewMatchMakingUser.Attributes.Emplace(Entry.Key.ToString(), FVariantData(Entry.Value.Data.GetInt64()));
-					break;
-				case ESchemaAttributeType::String:
-					NewMatchMakingUser.Attributes.Emplace(Entry.Key.ToString(), FVariantData(Entry.Value.Data.GetString()));
-					break;
-				}
+			case ESchemaAttributeType::Bool:
+				NewMatchMakingUser.Attributes.Emplace(Entry.Key.ToString(), FVariantData(Entry.Value.Data.GetBoolean()));
+				break;
+			case ESchemaAttributeType::Double:
+				NewMatchMakingUser.Attributes.Emplace(Entry.Key.ToString(), FVariantData(Entry.Value.Data.GetDouble()));
+				break;
+			case ESchemaAttributeType::Int64:
+				NewMatchMakingUser.Attributes.Emplace(Entry.Key.ToString(), FVariantData(Entry.Value.Data.GetInt64()));
+				break;
+			case ESchemaAttributeType::String:
+				NewMatchMakingUser.Attributes.Emplace(Entry.Key.ToString(), FVariantData(Entry.Value.Data.GetString()));
+				break;
 			}
-
-			MatchMakingUsers.Add(NewMatchMakingUser);
 		}
 
+		MatchMakingUsers.Add(NewMatchMakingUser);
+
 		FOnlineSessionSettings NewSessionSettings = BuildV1Settings(OpParams.SessionSettings);
 
 		/** We build a mock params struct for FindSessions and convert it to a V1 session search */
@@ -932,8 +930,8 @@ TOnlineAsyncOpHandle<FJoinSession> FSessionsOSSAdapter::JoinSession(FJoinSession
 			return;
 		}
 
-		const FUniqueNetIdRef LocalUserId = Services.Get<FAuthOSSAdapter>()->GetUniqueNetId(OpParams.LocalAccountId).ToSharedRef();
-		if (!LocalUserId->IsValid())
+		const FUniqueNetIdRef LocalAccountId = Services.Get<FAuthOSSAdapter>()->GetUniqueNetId(OpParams.LocalAccountId).ToSharedRef();
+		if (!LocalAccountId->IsValid())
 		{
 			Op.SetError(Errors::InvalidUser());
 			return;
@@ -1019,7 +1017,7 @@ TOnlineAsyncOpHandle<FJoinSession> FSessionsOSSAdapter::JoinSession(FJoinSession
 
 		SearchResult.Session = BuildV1Session(FoundSession);
 
-		SessionsInterface->JoinSession(*LocalUserId, OpParams.SessionName, SearchResult);
+		SessionsInterface->JoinSession(*LocalAccountId, OpParams.SessionName, SearchResult);
 
 		// TODO: Pending support for multiple local users
 	})
@@ -1028,16 +1026,16 @@ TOnlineAsyncOpHandle<FJoinSession> FSessionsOSSAdapter::JoinSession(FJoinSession
 	return Op->GetHandle();
 }
 
-TOnlineAsyncOpHandle<FAddSessionMembers> FSessionsOSSAdapter::AddSessionMembers(FAddSessionMembers::Params&& Params)
+TOnlineAsyncOpHandle<FAddSessionMember> FSessionsOSSAdapter::AddSessionMember(FAddSessionMember::Params&& Params)
 {
-	TOnlineAsyncOpRef<FAddSessionMembers> Op = GetOp<FAddSessionMembers>(MoveTemp(Params));
-	const FAddSessionMembers::Params& OpParams = Op->GetParams();
+	TOnlineAsyncOpRef<FAddSessionMember> Op = GetOp<FAddSessionMember>(MoveTemp(Params));
+	const FAddSessionMember::Params& OpParams = Op->GetParams();
 
-	Op->Then([this](TOnlineAsyncOp<FAddSessionMembers>& Op)
+	Op->Then([this](TOnlineAsyncOp<FAddSessionMember>& Op)
 	{
-		const FAddSessionMembers::Params& OpParams = Op.GetParams();
+		const FAddSessionMember::Params& OpParams = Op.GetParams();
 
-		FOnlineError StateCheck = CheckAddSessionMembersState(OpParams);
+		FOnlineError StateCheck = CheckAddSessionMemberState(OpParams);
 		if (StateCheck != Errors::Success())
 		{
 			Op.SetError(MoveTemp(StateCheck));
@@ -1048,7 +1046,7 @@ TOnlineAsyncOpHandle<FAddSessionMembers> FSessionsOSSAdapter::AddSessionMembers(
 		{
 			// Clearing the delegate needs to be the last thing we do, as modifying captured variables after it would result in a crash
 
-			if (TOnlineAsyncOpPtr<FAddSessionMembers> Op = WeakOp.Pin())
+			if (TOnlineAsyncOpPtr<FAddSessionMember> Op = WeakOp.Pin())
 			{
 				if (bWasSuccessful)
 				{
@@ -1064,17 +1062,12 @@ TOnlineAsyncOpHandle<FAddSessionMembers> FSessionsOSSAdapter::AddSessionMembers(
 		}));
 
 		FAuthOSSAdapter* Auth = Services.Get<FAuthOSSAdapter>();
-		TArray<FUniqueNetIdRef> TargetUserNetIds;
-		TArray<FAccountId> NewSessionMemberIds;
-		OpParams.NewSessionMembers.GenerateKeyArray(NewSessionMemberIds);
-		for (const FAccountId& TargetUser : NewSessionMemberIds)
+		TArray<FUniqueNetIdRef> TargetUserNetIds;		
+		FUniqueNetIdRef TargetUserNetId = Auth->GetUniqueNetId(OpParams.LocalAccountId).ToSharedRef();
+		if (TargetUserNetId->IsValid())
 		{
-			FUniqueNetIdRef TargetUserNetId = Auth->GetUniqueNetId(TargetUser).ToSharedRef();
-			if (TargetUserNetId->IsValid())
-			{
-				TargetUserNetIds.Add(TargetUserNetId);
-			}
-		}
+			TargetUserNetIds.Add(TargetUserNetId);
+		}	
 
 		SessionsInterface->RegisterPlayers(OpParams.SessionName, TargetUserNetIds);
 	})
@@ -1083,16 +1076,16 @@ TOnlineAsyncOpHandle<FAddSessionMembers> FSessionsOSSAdapter::AddSessionMembers(
 	return Op->GetHandle();
 }
 
-TOnlineAsyncOpHandle<FRemoveSessionMembers> FSessionsOSSAdapter::RemoveSessionMembers(FRemoveSessionMembers::Params&& Params)
+TOnlineAsyncOpHandle<FRemoveSessionMember> FSessionsOSSAdapter::RemoveSessionMember(FRemoveSessionMember::Params&& Params)
 {
-	TOnlineAsyncOpRef<FRemoveSessionMembers> Op = GetOp<FRemoveSessionMembers>(MoveTemp(Params));
-	const FRemoveSessionMembers::Params& OpParams = Op->GetParams();
+	TOnlineAsyncOpRef<FRemoveSessionMember> Op = GetOp<FRemoveSessionMember>(MoveTemp(Params));
+	const FRemoveSessionMember::Params& OpParams = Op->GetParams();
 
-	Op->Then([this](TOnlineAsyncOp<FRemoveSessionMembers>& Op)
+	Op->Then([this](TOnlineAsyncOp<FRemoveSessionMember>& Op)
 	{
-		const FRemoveSessionMembers::Params& OpParams = Op.GetParams();
+		const FRemoveSessionMember::Params& OpParams = Op.GetParams();
 
-		FOnlineError StateCheck = CheckRemoveSessionMembersState(OpParams);
+		FOnlineError StateCheck = CheckRemoveSessionMemberState(OpParams);
 		if (StateCheck != Errors::Success())
 		{
 			Op.SetError(MoveTemp(StateCheck));
@@ -1103,7 +1096,7 @@ TOnlineAsyncOpHandle<FRemoveSessionMembers> FSessionsOSSAdapter::RemoveSessionMe
 		{
 			// Clearing the delegate needs to be the last thing we do, as modifying captured variables after it would result in a crash
 
-			if (TOnlineAsyncOpPtr<FRemoveSessionMembers> Op = WeakOp.Pin())
+			if (TOnlineAsyncOpPtr<FRemoveSessionMember> Op = WeakOp.Pin())
 			{
 				if (bWasSuccessful)
 				{
@@ -1120,13 +1113,10 @@ TOnlineAsyncOpHandle<FRemoveSessionMembers> FSessionsOSSAdapter::RemoveSessionMe
 
 		FAuthOSSAdapter* Auth = Services.Get<FAuthOSSAdapter>();
 		TArray<FUniqueNetIdRef> TargetUserNetIds;
-		for (const FAccountId& TargetUser : OpParams.SessionMemberIds)
+		FUniqueNetIdRef TargetUserNetId = Auth->GetUniqueNetId(OpParams.LocalAccountId).ToSharedRef();
+		if (TargetUserNetId->IsValid())
 		{
-			FUniqueNetIdRef TargetUserNetId = Auth->GetUniqueNetId(TargetUser).ToSharedRef();
-			if (TargetUserNetId->IsValid())
-			{
-				TargetUserNetIds.Add(TargetUserNetId);
-			}
+			TargetUserNetIds.Add(TargetUserNetId);
 		}
 
 		SessionsInterface->UnregisterPlayers(OpParams.SessionName, TargetUserNetIds);
@@ -1154,8 +1144,8 @@ TOnlineAsyncOpHandle<FSendSessionInvite> FSessionsOSSAdapter::SendSessionInvite(
 
 		FAuthOSSAdapter* Auth = Services.Get<FAuthOSSAdapter>();
 
-		const FUniqueNetIdRef LocalUserId = Auth->GetUniqueNetId(OpParams.LocalAccountId).ToSharedRef();
-		if (!LocalUserId->IsValid())
+		const FUniqueNetIdRef LocalAccountId = Auth->GetUniqueNetId(OpParams.LocalAccountId).ToSharedRef();
+		if (!LocalAccountId->IsValid())
 		{
 			Op.SetError(Errors::InvalidUser());
 			return;
@@ -1171,7 +1161,7 @@ TOnlineAsyncOpHandle<FSendSessionInvite> FSessionsOSSAdapter::SendSessionInvite(
 			}
 		}
 
-		SessionsInterface->SendSessionInviteToFriends(*LocalUserId, OpParams.SessionName, TargetUserNetIds);
+		SessionsInterface->SendSessionInviteToFriends(*LocalAccountId, OpParams.SessionName, TargetUserNetIds);
 
 		Op.SetResult({ });
 	})
diff --git a/Engine/Plugins/Online/Experimental/OnlineServicesOSSAdapter/Source/Public/Online/SessionsOSSAdapter.h b/Engine/Plugins/Online/Experimental/OnlineServicesOSSAdapter/Source/Public/Online/SessionsOSSAdapter.h
--- a/Engine/Plugins/Online/Experimental/OnlineServicesOSSAdapter/Source/Public/Online/SessionsOSSAdapter.h
+++ b/Engine/Plugins/Online/Experimental/OnlineServicesOSSAdapter/Source/Public/Online/SessionsOSSAdapter.h
@@ -69,8 +69,8 @@ public:
 	virtual TOnlineAsyncOpHandle<FFindSessions> FindSessions(FFindSessions::Params&& Params) override;
 	virtual TOnlineAsyncOpHandle<FStartMatchmaking> StartMatchmaking(FStartMatchmaking::Params&& Params) override;
 	virtual TOnlineAsyncOpHandle<FJoinSession> JoinSession(FJoinSession::Params&& Params) override;
-	virtual TOnlineAsyncOpHandle<FAddSessionMembers> AddSessionMembers(FAddSessionMembers::Params&& Params) override;
-	virtual TOnlineAsyncOpHandle<FRemoveSessionMembers> RemoveSessionMembers(FRemoveSessionMembers::Params&& Params) override;
+	virtual TOnlineAsyncOpHandle<FAddSessionMember> AddSessionMember(FAddSessionMember::Params&& Params) override;
+	virtual TOnlineAsyncOpHandle<FRemoveSessionMember> RemoveSessionMember(FRemoveSessionMember::Params&& Params) override;
 	virtual TOnlineAsyncOpHandle<FSendSessionInvite> SendSessionInvite(FSendSessionInvite::Params&& Params) override;
 	virtual TOnlineAsyncOpHandle<FRegisterPlayers> RegisterPlayers(FRegisterPlayers::Params&& Params) override;
 	virtual TOnlineAsyncOpHandle<FUnregisterPlayers> UnregisterPlayers(FUnregisterPlayers::Params&& Params) override;
