Refactored GetAllSessions to return all sessions a given user is in, instead of all sessions in the system

diff --git a/Engine/Plugins/Online/Experimental/OnlineServices/Source/OnlineServicesCommon/Private/Online/SessionsCommon.cpp b/Engine/Plugins/Online/Experimental/OnlineServices/Source/OnlineServicesCommon/Private/Online/SessionsCommon.cpp
--- a/Engine/Plugins/Online/Experimental/OnlineServices/Source/OnlineServicesCommon/Private/Online/SessionsCommon.cpp
+++ b/Engine/Plugins/Online/Experimental/OnlineServices/Source/OnlineServicesCommon/Private/Online/SessionsCommon.cpp
@@ -39,14 +39,24 @@ namespace UE::Online {
 
 	TOnlineResult<FGetAllSessions> FSessionsCommon::GetAllSessions(FGetAllSessions::Params&& Params) const
 	{
-		FGetAllSessions::Result Result;
+		if (const TArray<FName>* UserSessions = NamedSessionUserMap.Find(Params.LocalAccountId))
+		{
+			FGetAllSessions::Result Result;
+
+			for (const FName& SessionName : *UserSessions)
+			{
+				const FOnlineSessionIdHandle& SessionId = LocalSessionsByName.FindChecked(SessionName);
+				const TSharedRef<FSessionCommon>& Session = AllSessionsById.FindChecked(SessionId);
+
+				Result.Sessions.Add(Session);
+			}
 
-		for (const TPair<FOnlineSessionIdHandle, TSharedRef<FSessionCommon>>& SessionPair : AllSessionsById)
+			return TOnlineResult<FGetAllSessions>(MoveTemp(Result));
+		}
+		else
 		{
-			Result.Sessions.Add(SessionPair.Value);
+			return TOnlineResult<FGetAllSessions>(Errors::InvalidUser());
 		}
-
-		return TOnlineResult<FGetAllSessions>(MoveTemp(Result));
 	}
 
 	TOnlineResult<FGetSessionByName> FSessionsCommon::GetSessionByName(FGetSessionByName::Params&& Params) const
diff --git a/Engine/Plugins/Online/Experimental/OnlineServices/Source/OnlineServicesInterface/Public/Online/Sessions.h b/Engine/Plugins/Online/Experimental/OnlineServices/Source/OnlineServicesInterface/Public/Online/Sessions.h
--- a/Engine/Plugins/Online/Experimental/OnlineServices/Source/OnlineServicesInterface/Public/Online/Sessions.h
+++ b/Engine/Plugins/Online/Experimental/OnlineServices/Source/OnlineServicesInterface/Public/Online/Sessions.h
@@ -211,7 +211,7 @@ struct FGetAllSessions
 
 	struct Params
 	{
-
+		FAccountId LocalAccountId;
 	};
 
 	struct Result
@@ -672,7 +672,7 @@ class ISessions
 {
 public:
 	/**
-	 * Get an array of all session objects.
+	 * Gets an array of references to all the sessions the given user is part of.
 	 * 
 	 * @params Parameters for the GetAllSessions call
 	 * return
@@ -952,6 +952,7 @@ BEGIN_ONLINE_STRUCT_META(FSessionInvite)
 END_ONLINE_STRUCT_META()
 
 BEGIN_ONLINE_STRUCT_META(FGetAllSessions::Params)
+	ONLINE_STRUCT_FIELD(FGetAllSessions::Params, LocalAccountId)
 END_ONLINE_STRUCT_META()
 
 BEGIN_ONLINE_STRUCT_META(FGetAllSessions::Result)
diff --git a/Engine/Plugins/Online/Experimental/OnlineServicesNull/Source/Private/Online/OnlineServicesNull.cpp b/Engine/Plugins/Online/Experimental/OnlineServicesNull/Source/Private/Online/OnlineServicesNull.cpp
--- a/Engine/Plugins/Online/Experimental/OnlineServicesNull/Source/Private/Online/OnlineServicesNull.cpp
+++ b/Engine/Plugins/Online/Experimental/OnlineServicesNull/Source/Private/Online/OnlineServicesNull.cpp
@@ -90,8 +90,7 @@ TOnlineResult<FGetResolvedConnectString> FOnlineServicesNull::GetResolvedConnect
 		ISessionsPtr SessionsPtr = GetSessionsInterface();
 		check(SessionsPtr);
 
-		FGetAllSessions::Params GetAllSessionsParams;
-		TOnlineResult<FGetAllSessions> JoinedSessions = SessionsPtr->GetAllSessions(MoveTemp(GetAllSessionsParams));
+		TOnlineResult<FGetAllSessions> JoinedSessions = SessionsPtr->GetAllSessions({ Params.LocalAccountId });
 		if (JoinedSessions.IsOk())
 		{
 			for (TSharedRef<const ISession>& Session : JoinedSessions.GetOkValue().Sessions)
