OSSv2:
- Rename FOnlineSessionIdHandle to FOnlineSessionId, not renaming to FSessionId because the name already existing and I think it might be good to keep Online to distinguish it from network session for the replication;
- Rename corresponding variable SessionIdHandle/Handle to SessionId.

diff --git a/Engine/Plugins/Online/Experimental/OnlineServices/Source/OnlineServicesCommon/Private/Online/SessionsCommon.cpp b/Engine/Plugins/Online/Experimental/OnlineServices/Source/OnlineServicesCommon/Private/Online/SessionsCommon.cpp
--- a/Engine/Plugins/Online/Experimental/OnlineServices/Source/OnlineServicesCommon/Private/Online/SessionsCommon.cpp
+++ b/Engine/Plugins/Online/Experimental/OnlineServices/Source/OnlineServicesCommon/Private/Online/SessionsCommon.cpp
@@ -41,18 +41,18 @@ namespace UE::Online {
 	{
 		if (const TArray<FName>* UserSessions = NamedSessionUserMap.Find(Params.LocalAccountId))
 		{
-			FGetAllSessions::Result Result;
+		FGetAllSessions::Result Result;
 
 			for (const FName& SessionName : *UserSessions)
-			{
-				const FOnlineSessionIdHandle& SessionId = LocalSessionsByName.FindChecked(SessionName);
+		{
+				const FOnlineSessionId& SessionId = LocalSessionsByName.FindChecked(SessionName);
 				const TSharedRef<FSessionCommon>& Session = AllSessionsById.FindChecked(SessionId);
 
 				Result.Sessions.Add(Session);
-			}
-
-			return TOnlineResult<FGetAllSessions>(MoveTemp(Result));
 		}
+
+		return TOnlineResult<FGetAllSessions>(MoveTemp(Result));
+	}
 		else
 		{
 			return TOnlineResult<FGetAllSessions>(Errors::InvalidUser());
@@ -74,7 +74,7 @@ namespace UE::Online {
 
 	TOnlineResult<FGetSessionById> FSessionsCommon::GetSessionById(FGetSessionById::Params&& Params) const
 	{
-		TOnlineResult<FGetMutableSessionById> GetMutableSessionByIdResult = GetMutableSessionById({ Params.IdHandle });
+		TOnlineResult<FGetMutableSessionById> GetMutableSessionByIdResult = GetMutableSessionById({ Params.SessionId });
 		if (GetMutableSessionByIdResult.IsOk())
 		{
 			return TOnlineResult<FGetSessionById>({ GetMutableSessionByIdResult.GetOkValue().Session });
@@ -87,7 +87,7 @@ namespace UE::Online {
 
 	TOnlineResult<FGetPresenceSession> FSessionsCommon::GetPresenceSession(FGetPresenceSession::Params&& Params) const
 	{
-		if (const FOnlineSessionIdHandle* PresenceSessionId = PresenceSessionsUserMap.Find(Params.LocalAccountId))
+		if (const FOnlineSessionId* PresenceSessionId = PresenceSessionsUserMap.Find(Params.LocalAccountId))
 		{
 			return TOnlineResult<FGetPresenceSession>({ AllSessionsById.FindChecked(*PresenceSessionId) });
 		}
@@ -99,7 +99,7 @@ namespace UE::Online {
 
 	TOnlineResult<FSetPresenceSession> FSessionsCommon::SetPresenceSession(FSetPresenceSession::Params&& Params)
 	{
-		FOnlineSessionIdHandle& PresenceSessionId = PresenceSessionsUserMap.FindOrAdd(Params.LocalAccountId);
+		FOnlineSessionId& PresenceSessionId = PresenceSessionsUserMap.FindOrAdd(Params.LocalAccountId);
 		PresenceSessionId = Params.SessionId;
 
 		return TOnlineResult<FSetPresenceSession>(FSetPresenceSession::Result{ });
@@ -336,11 +336,11 @@ namespace UE::Online {
 
 	TOnlineResult<FGetMutableSessionByName> FSessionsCommon::GetMutableSessionByName(FGetMutableSessionByName::Params&& Params) const
 	{
-		if (const FOnlineSessionIdHandle* SessionIdHandle = LocalSessionsByName.Find(Params.LocalName))
+		if (const FOnlineSessionId* SessionId = LocalSessionsByName.Find(Params.LocalName))
 		{
-			check(AllSessionsById.Contains(*SessionIdHandle));
+			check(AllSessionsById.Contains(*SessionId));
 
-			return TOnlineResult<FGetMutableSessionByName>({ AllSessionsById.FindChecked(*SessionIdHandle) });
+			return TOnlineResult<FGetMutableSessionByName>({ AllSessionsById.FindChecked(*SessionId) });
 		}
 		else
 		{
@@ -350,7 +350,7 @@ namespace UE::Online {
 
 	TOnlineResult<FGetMutableSessionById> FSessionsCommon::GetMutableSessionById(FGetMutableSessionById::Params&& Params) const
 	{
-		if (const TSharedRef<FSessionCommon>* FoundSession = AllSessionsById.Find(Params.IdHandle))
+		if (const TSharedRef<FSessionCommon>* FoundSession = AllSessionsById.Find(Params.SessionId))
 		{
 			return TOnlineResult<FGetMutableSessionById>({ *FoundSession });
 		}
@@ -372,7 +372,7 @@ namespace UE::Online {
 	{
 		AllSessionsById.Emplace(Session->SessionId, Session);
 
-		TArray<FOnlineSessionIdHandle>& SearchResults = SearchResultsUserMap.FindOrAdd(LocalAccountId);
+		TArray<FOnlineSessionId>& SearchResults = SearchResultsUserMap.FindOrAdd(LocalAccountId);
 		SearchResults.Add(Session->SessionId);
 	}
 
@@ -383,7 +383,7 @@ namespace UE::Online {
 		AddSessionReferences(Session->SessionId, SessionName, LocalAccountId, bIsPresenceSession);
 	}
 
-	void FSessionsCommon::AddSessionReferences(const FOnlineSessionIdHandle SessionId, const FName& SessionName, const FAccountId& LocalAccountId, bool bIsPresenceSession)
+	void FSessionsCommon::AddSessionReferences(const FOnlineSessionId SessionId, const FName& SessionName, const FAccountId& LocalAccountId, bool bIsPresenceSession)
 	{
 		LocalSessionsByName.Emplace(SessionName, SessionId);
 
@@ -395,7 +395,7 @@ namespace UE::Online {
 		}
 	}
 
-	void FSessionsCommon::ClearSessionInvitesForSession(const FAccountId& LocalAccountId, const FOnlineSessionIdHandle SessionId)
+	void FSessionsCommon::ClearSessionInvitesForSession(const FAccountId& LocalAccountId, const FOnlineSessionId SessionId)
 	{
 		if (TMap<FOnlineSessionInviteIdHandle, TSharedRef<FSessionInvite>>* UserMap = SessionInvitesUserMap.Find(LocalAccountId))
 		{
@@ -415,7 +415,7 @@ namespace UE::Online {
 		}
 	}
 
-	void FSessionsCommon::ClearSessionReferences(const FOnlineSessionIdHandle SessionId, const FName& SessionName, const FAccountId& LocalAccountId, bool bIsPresenceSession)
+	void FSessionsCommon::ClearSessionReferences(const FOnlineSessionId SessionId, const FName& SessionName, const FAccountId& LocalAccountId, bool bIsPresenceSession)
 	{
 		NamedSessionUserMap.FindChecked(LocalAccountId).Remove(SessionName);
 
@@ -442,10 +442,10 @@ namespace UE::Online {
 		LocalSessionsByName.Remove(SessionName);
 	}
 
-	void FSessionsCommon::ClearSessionById(const FOnlineSessionIdHandle& SessionId)
+	void FSessionsCommon::ClearSessionById(const FOnlineSessionId& SessionId)
 	{
 		// PresenceSessionsUserMap is not evaluated, since any session there would also be in LocalSessionsByName
-		for (const TPair<FName, FOnlineSessionIdHandle>& Entry : LocalSessionsByName)
+		for (const TPair<FName, FOnlineSessionId>& Entry : LocalSessionsByName)
 		{
 			if (Entry.Value == SessionId)
 			{
@@ -464,7 +464,7 @@ namespace UE::Online {
 			}			
 		}
 
-		for (const TPair<FAccountId, TArray<FOnlineSessionIdHandle>>& Entry : SearchResultsUserMap)
+		for (const TPair<FAccountId, TArray<FOnlineSessionId>>& Entry : SearchResultsUserMap)
 		{
 			if (Entry.Value.Contains(SessionId))
 			{
@@ -723,7 +723,7 @@ namespace UE::Online {
 
 		if (Params.SessionSettings.bPresenceEnabled)
 		{
-			for (const TPair<FName, FOnlineSessionIdHandle>& Entry : LocalSessionsByName)
+			for (const TPair<FName, FOnlineSessionId>& Entry : LocalSessionsByName)
 			{
 				TOnlineResult<FGetSessionById> GetSessionByIdResult = GetSessionById({ Params.LocalAccountId, Entry.Value });
 				if (GetSessionByIdResult.IsOk())
diff --git a/Engine/Plugins/Online/Experimental/OnlineServices/Source/OnlineServicesCommon/Public/Online/SessionsCommon.h b/Engine/Plugins/Online/Experimental/OnlineServices/Source/OnlineServicesCommon/Public/Online/SessionsCommon.h
--- a/Engine/Plugins/Online/Experimental/OnlineServices/Source/OnlineServicesCommon/Public/Online/SessionsCommon.h
+++ b/Engine/Plugins/Online/Experimental/OnlineServices/Source/OnlineServicesCommon/Public/Online/SessionsCommon.h
@@ -16,9 +16,9 @@ class FOnlineSessionIdStringRegistry : public IOnlineSessionIdRegistry
 {
 public:
 	// Begin IOnlineSessionIdRegistry
-	virtual inline FString ToLogString(const FOnlineSessionIdHandle& Handle) const override
+	virtual inline FString ToLogString(const FOnlineSessionId& SessionId) const override
 	{
-		FString IdValue = BasicRegistry.FindIdValue(Handle);
+		FString IdValue = BasicRegistry.FindIdValue(SessionId);
 
 		if (IdValue.Len() == 0)
 		{
@@ -28,9 +28,9 @@ public:
 		return IdValue;
 	};
 
-	virtual inline TArray<uint8> ToReplicationData(const FOnlineSessionIdHandle& Handle) const override
+	virtual inline TArray<uint8> ToReplicationData(const FOnlineSessionId& SessionId) const override
 	{
-		const FString IdValue = BasicRegistry.FindIdValue(Handle);
+		const FString IdValue = BasicRegistry.FindIdValue(SessionId);
 		const FTCHARToUTF8 IdValueUtf8(*IdValue);
 
 		TArray<uint8> ReplicationData;
@@ -41,7 +41,7 @@ public:
 		return ReplicationData;
 	}
 
-	virtual inline FOnlineSessionIdHandle FromReplicationData(const TArray<uint8>& ReplicationData) override
+	virtual inline FOnlineSessionId FromReplicationData(const TArray<uint8>& ReplicationData) override
 	{
 		const FUTF8ToTCHAR IdValueTCHAR((char*)ReplicationData.GetData(), ReplicationData.Num());
 		const FString IdValue = FString(IdValueTCHAR.Length(), IdValueTCHAR.Get());
@@ -51,13 +51,13 @@ public:
 			return BasicRegistry.FindOrAddHandle(IdValue);
 		}
 
-		return FOnlineSessionIdHandle();
+		return FOnlineSessionId();
 	}
 	// End IOnlineSessionIdRegistry
 
-	inline bool IsSessionIdExpired(const FOnlineSessionIdHandle& InHandle) const
+	inline bool IsSessionIdExpired(const FOnlineSessionId& InSessionId) const
 	{
-		return BasicRegistry.FindIdValue(InHandle).IsEmpty();
+		return BasicRegistry.FindIdValue(InSessionId).IsEmpty();
 	}
 
 	FOnlineSessionIdStringRegistry(EOnlineServices OnlineServicesType)
@@ -136,7 +136,7 @@ public:
 
 	// ISession
 	virtual const FAccountId GetOwnerAccountId() const override				{ return OwnerAccountId; }
-	virtual const FOnlineSessionIdHandle GetSessionId() const override		{ return SessionId; }
+	virtual const FOnlineSessionId GetSessionId() const override		{ return SessionId; }
 	virtual const FSessionSettings GetSessionSettings() const override		{ return SessionSettings; }
 
 	virtual FString ToLogString() const override							{ return FString(); } // TODO: Implement after completing refactor
@@ -146,7 +146,7 @@ public:
 	FAccountId OwnerAccountId;
 
 	/** The id for the session, platform dependent */
-	FOnlineSessionIdHandle SessionId;
+	FOnlineSessionId SessionId;
 
 	/** Set of session properties that can be altered by the session owner */
 	FSessionSettings SessionSettings;
@@ -175,7 +175,7 @@ struct FGetMutableSessionById
 
 	struct Params
 	{
-		FOnlineSessionIdHandle IdHandle;
+		FOnlineSessionId SessionId;
 	};
 
 	struct Result
@@ -231,10 +231,10 @@ protected:
 	void AddSessionInvite(const TSharedRef<FSessionInvite> SessionInvite, const TSharedRef<FSessionCommon> Session, const FAccountId& LocalAccountId);
 	void AddSearchResult(const TSharedRef<FSessionCommon> Session, const FAccountId& LocalAccountId);
 	void AddSessionWithReferences(const TSharedRef<FSessionCommon> Session, const FName& SessionName, const FAccountId& LocalAccountId, bool bIsPresenceSession);
-	void AddSessionReferences(const FOnlineSessionIdHandle SessionId, const FName& SessionName, const FAccountId& LocalAccountId, bool bIsPresenceSession);
+	void AddSessionReferences(const FOnlineSessionId SessionId, const FName& SessionName, const FAccountId& LocalAccountId, bool bIsPresenceSession);
 
-	void ClearSessionInvitesForSession(const FAccountId& LocalAccountId, const FOnlineSessionIdHandle SessionId);
-	void ClearSessionReferences(const FOnlineSessionIdHandle SessionId, const FName& SessionName, const FAccountId& LocalAccountId, bool bIsPresenceSession);
+	void ClearSessionInvitesForSession(const FAccountId& LocalAccountId, const FOnlineSessionId SessionId);
+	void ClearSessionReferences(const FOnlineSessionId SessionId, const FName& SessionName, const FAccountId& LocalAccountId, bool bIsPresenceSession);
 
 	TOnlineResult<FAddSessionMember> AddSessionMemberImpl(const FAddSessionMember::Params& Params);
 	TOnlineResult<FRemoveSessionMember> RemoveSessionMemberImpl(const FRemoveSessionMember::Params& Params);
@@ -274,7 +274,7 @@ private:
 	TOptional<FOnlineError> CheckSessionExistsByName(const FAccountId& LocalAccountId, const FName& SessionName);
 
 	void ClearSessionByName(const FName& SessionName);
-	void ClearSessionById(const FOnlineSessionIdHandle& SessionId);
+	void ClearSessionById(const FOnlineSessionId& SessionId);
 
 protected:
 
@@ -291,22 +291,22 @@ protected:
 	TMap<FAccountId, TArray<FName>> NamedSessionUserMap;
 
 	/** Map of sessions that local users are part of, indexed by their local name */
-	TMap<FName, FOnlineSessionIdHandle> LocalSessionsByName;
+	TMap<FName, FOnlineSessionId> LocalSessionsByName;
 
 	/** Map of sessions that local users have set as their presence session to appear in the platform UI. A user may not have set any session as their presence session. */
-	TMap<FAccountId, FOnlineSessionIdHandle> PresenceSessionsUserMap;
+	TMap<FAccountId, FOnlineSessionId> PresenceSessionsUserMap;
 
 	/** Cache for received session invites, mapped per user */
 	TMap<FAccountId, TMap<FOnlineSessionInviteIdHandle, TSharedRef<FSessionInvite>>> SessionInvitesUserMap;
 
 	/** Cache for the last set of session search results, mapped per user */
-	TMap<FAccountId, TArray<FOnlineSessionIdHandle>> SearchResultsUserMap;
+	TMap<FAccountId, TArray<FOnlineSessionId>> SearchResultsUserMap;
 
 	/** Handle for an ongoing session search operation, mapped per user */
 	TMap<FAccountId, TSharedRef<TOnlineAsyncOp<FFindSessions>>> CurrentSessionSearchHandlesUserMap;
 
 	/** Set of every distinct FSession found, indexed by Id */
-	TMap<FOnlineSessionIdHandle, TSharedRef<FSessionCommon>> AllSessionsById;
+	TMap<FOnlineSessionId, TSharedRef<FSessionCommon>> AllSessionsById;
 };
 
 namespace Meta {
@@ -320,7 +320,7 @@ BEGIN_ONLINE_STRUCT_META(FGetMutableSessionByName::Result)
 END_ONLINE_STRUCT_META()
 
 BEGIN_ONLINE_STRUCT_META(FGetMutableSessionById::Params)
-	ONLINE_STRUCT_FIELD(FGetMutableSessionById::Params, IdHandle)
+	ONLINE_STRUCT_FIELD(FGetMutableSessionById::Params, SessionId)
 END_ONLINE_STRUCT_META()
 
 BEGIN_ONLINE_STRUCT_META(FGetMutableSessionById::Result)
diff --git a/Engine/Plugins/Online/Experimental/OnlineServices/Source/OnlineServicesCommonEngineUtils/Private/Online/SessionsLAN.cpp b/Engine/Plugins/Online/Experimental/OnlineServices/Source/OnlineServicesCommonEngineUtils/Private/Online/SessionsLAN.cpp
--- a/Engine/Plugins/Online/Experimental/OnlineServices/Source/OnlineServicesCommonEngineUtils/Private/Online/SessionsLAN.cpp
+++ b/Engine/Plugins/Online/Experimental/OnlineServices/Source/OnlineServicesCommonEngineUtils/Private/Online/SessionsLAN.cpp
@@ -25,7 +25,7 @@ FOnlineSessionIdRegistryLAN& FOnlineSessionIdRegistryLAN::GetChecked(EOnlineServ
 	return *SessionIdRegistry;
 }
 
-FOnlineSessionIdHandle FOnlineSessionIdRegistryLAN::GetNextSessionId()
+FOnlineSessionId FOnlineSessionIdRegistryLAN::GetNextSessionId()
 {
 	return BasicRegistry.FindOrAddHandle(FGuid::NewGuid().ToString());
 }
@@ -436,7 +436,7 @@ void FSessionsLAN::StopLANSession()
 void FSessionsLAN::OnValidQueryPacketReceived(uint8* PacketData, int32 PacketLength, uint64 ClientNonce)
 {
 	// Iterate through all registered sessions and respond for each one that can be joinable
-	for (const TPair<FOnlineSessionIdHandle, TSharedRef<FSessionCommon>>& Entry : AllSessionsById)
+	for (const TPair<FOnlineSessionId, TSharedRef<FSessionCommon>>& Entry : AllSessionsById)
 	{
 		const TSharedRef<FSessionCommon>& Session = Entry.Value;
 
@@ -481,7 +481,7 @@ void FSessionsLAN::OnLANSearchTimeout(const FAccountId LocalAccountId)
 
 	UE_LOG(LogTemp, Warning, TEXT("[FSessionsLAN::OnLANSearchTimeout] %d sessions found!"), SearchResultsUserMap.FindChecked(LocalAccountId).Num());
 
-	TArray<FOnlineSessionIdHandle>& FoundSessionIds = SearchResultsUserMap.FindChecked(LocalAccountId);
+	TArray<FOnlineSessionId>& FoundSessionIds = SearchResultsUserMap.FindChecked(LocalAccountId);
 
 	CurrentSessionSearchHandlesUserMap.FindChecked(LocalAccountId)->SetResult({ FoundSessionIds });
 
diff --git a/Engine/Plugins/Online/Experimental/OnlineServices/Source/OnlineServicesCommonEngineUtils/Public/Online/SessionsLAN.h b/Engine/Plugins/Online/Experimental/OnlineServices/Source/OnlineServicesCommonEngineUtils/Public/Online/SessionsLAN.h
--- a/Engine/Plugins/Online/Experimental/OnlineServices/Source/OnlineServicesCommonEngineUtils/Public/Online/SessionsLAN.h
+++ b/Engine/Plugins/Online/Experimental/OnlineServices/Source/OnlineServicesCommonEngineUtils/Public/Online/SessionsLAN.h
@@ -15,7 +15,7 @@ class ONLINESERVICESCOMMONENGINEUTILS_API FOnlineSessionIdRegistryLAN : public F
 public:
 	static FOnlineSessionIdRegistryLAN& GetChecked(EOnlineServices ServicesType);
 
-	FOnlineSessionIdHandle GetNextSessionId();
+	FOnlineSessionId GetNextSessionId();
 
 protected:
 	FOnlineSessionIdRegistryLAN(EOnlineServices ServicesType);
diff --git a/Engine/Plugins/Online/Experimental/OnlineServices/Source/OnlineServicesInterface/Public/Online/OnlineServices.h b/Engine/Plugins/Online/Experimental/OnlineServices/Source/OnlineServicesInterface/Public/Online/OnlineServices.h
--- a/Engine/Plugins/Online/Experimental/OnlineServices/Source/OnlineServicesInterface/Public/Online/OnlineServices.h
+++ b/Engine/Plugins/Online/Experimental/OnlineServices/Source/OnlineServicesInterface/Public/Online/OnlineServices.h
@@ -37,7 +37,7 @@ struct FGetResolvedConnectString
 	{
 		FAccountId LocalAccountId;
 		FLobbyId LobbyId;
-		FOnlineSessionIdHandle SessionId;
+		FOnlineSessionId SessionId;
 		FName PortType;
 	};
 
diff --git a/Engine/Plugins/Online/Experimental/OnlineServices/Source/OnlineServicesInterface/Public/Online/Sessions.h b/Engine/Plugins/Online/Experimental/OnlineServices/Source/OnlineServicesInterface/Public/Online/Sessions.h
--- a/Engine/Plugins/Online/Experimental/OnlineServices/Source/OnlineServicesInterface/Public/Online/Sessions.h
+++ b/Engine/Plugins/Online/Experimental/OnlineServices/Source/OnlineServicesInterface/Public/Online/Sessions.h
@@ -182,7 +182,7 @@ class ISession
 {
 public:
 	virtual const FAccountId GetOwnerAccountId() const = 0;
-	virtual const FOnlineSessionIdHandle GetSessionId() const = 0;
+	virtual const FOnlineSessionId GetSessionId() const = 0;
 	virtual const FSessionSettings GetSessionSettings() const = 0;
 	virtual FString ToLogString() const = 0;
 };
@@ -200,7 +200,7 @@ struct FSessionInvite
 	FOnlineSessionInviteIdHandle InviteId;
 
 	/* Pointer to the session information */
-	FOnlineSessionIdHandle SessionId;
+	FOnlineSessionId SessionId;
 
 	// TODO: Default constructor will be deleted after we cache invites
 };
@@ -245,7 +245,7 @@ struct FGetSessionById
 	{
 		FAccountId LocalAccountId;
 
-		FOnlineSessionIdHandle IdHandle;
+		FOnlineSessionId SessionId;
 	};
 
 	struct Result
@@ -279,7 +279,7 @@ struct FSetPresenceSession
 	{
 		FAccountId LocalAccountId;
 
-		FOnlineSessionIdHandle SessionId;
+		FOnlineSessionId SessionId;
 	};
 
 	struct Result
@@ -392,12 +392,12 @@ struct FFindSessions
 		TOptional<FAccountId> TargetUser;
 
 		/* Find join info for the target session id. */
-		TOptional<FOnlineSessionIdHandle> SessionId;
+		TOptional<FOnlineSessionId> SessionId;
 	};
 
 	struct Result
 	{
-		TArray<FOnlineSessionIdHandle> FoundSessionIds;
+		TArray<FOnlineSessionId> FoundSessionIds;
 	};
 };
 
@@ -442,7 +442,7 @@ struct FJoinSession
 		FName SessionName;
 
 		/* Id handle for the session to be joined. To be retrieved via session search or invite */
-		FOnlineSessionIdHandle SessionId;
+		FOnlineSessionId SessionId;
 
 		/* Information for the local user who will join the session */
 		FSessionMember SessionMemberData;
@@ -614,7 +614,7 @@ struct FSessionJoined
 	FAccountId LocalAccountId;
 
 	/* Id for the session joined. */
-	FOnlineSessionIdHandle SessionId;
+	FOnlineSessionId SessionId;
 };
 
 struct FSessionLeft
@@ -660,7 +660,7 @@ struct FUISessionJoinRequested
 	FAccountId LocalAccountId;
 
 	/** The id for the session the local user requested to join, or the online error if there was a failure retrieving it */
-	TResult<FOnlineSessionIdHandle, FOnlineError> Result;
+	TResult<FOnlineSessionId, FOnlineError> Result;
 
 	/** Join request source */
 	EUISessionJoinRequestedSource JoinRequestedSource = EUISessionJoinRequestedSource::Unspecified;
@@ -969,7 +969,7 @@ END_ONLINE_STRUCT_META()
 
 BEGIN_ONLINE_STRUCT_META(FGetSessionById::Params)
 	ONLINE_STRUCT_FIELD(FGetSessionById::Params, LocalAccountId),
-	ONLINE_STRUCT_FIELD(FGetSessionById::Params, IdHandle)
+	ONLINE_STRUCT_FIELD(FGetSessionById::Params, SessionId)
 END_ONLINE_STRUCT_META()
 
 BEGIN_ONLINE_STRUCT_META(FGetSessionById::Result)
diff --git a/Engine/Plugins/Online/Experimental/OnlineServicesEOSGS/Source/Private/Online/SessionsEOSGS.cpp b/Engine/Plugins/Online/Experimental/OnlineServicesEOSGS/Source/Private/Online/SessionsEOSGS.cpp
--- a/Engine/Plugins/Online/Experimental/OnlineServicesEOSGS/Source/Private/Online/SessionsEOSGS.cpp
+++ b/Engine/Plugins/Online/Experimental/OnlineServicesEOSGS/Source/Private/Online/SessionsEOSGS.cpp
@@ -332,7 +332,7 @@ void FSessionsEOSGS::HandleSessionInviteAccepted(const EOS_Sessions_SessionInvit
 
 					FUISessionJoinRequested Event{
 									Result.GetOkValue().LocalAccountId,
-									TResult<FOnlineSessionIdHandle, FOnlineError>(Result.GetOkValue().Session->SessionId),
+									TResult<FOnlineSessionId, FOnlineError>(Result.GetOkValue().Session->SessionId),
 						EUISessionJoinRequestedSource::FromInvitation
 					};
 
@@ -362,7 +362,7 @@ void FSessionsEOSGS::HandleJoinSessionAccepted(const EOS_Sessions_JoinSessionAcc
 
 					FUISessionJoinRequested Event{
 						Result.GetOkValue().LocalAccountId,
-						TResult<FOnlineSessionIdHandle, FOnlineError>(Result.GetOkValue().Session->SessionId),
+						TResult<FOnlineSessionId, FOnlineError>(Result.GetOkValue().Session->SessionId),
 			EUISessionJoinRequestedSource::Unspecified
 					};
 
@@ -1188,7 +1188,7 @@ void FSessionsEOSGS::SetSessionSearchParameters(FSessionSearchHandleEOSGS& Sessi
 	}
 }
 
-void FSessionsEOSGS::SetSessionSearchSessionId(FSessionSearchHandleEOSGS& SessionSearchHandle, const FOnlineSessionIdHandle& SessionId)
+void FSessionsEOSGS::SetSessionSearchSessionId(FSessionSearchHandleEOSGS& SessionSearchHandle, const FOnlineSessionId& SessionId)
 {
 	EOS_SessionSearch_SetSessionIdOptions Options = { };
 	Options.ApiVersion = EOS_SESSIONSEARCH_SETSESSIONID_API_LATEST;
@@ -1238,7 +1238,7 @@ void FSessionsEOSGS::WriteSessionSearchHandle(FSessionSearchHandleEOSGS& Session
 	}
 }
 
-FOnlineSessionIdHandle FSessionsEOSGS::CreateSessionId(const FString& SessionId)
+FOnlineSessionId FSessionsEOSGS::CreateSessionId(const FString& SessionId)
 {
 	return FOnlineSessionIdRegistryEOSGS::Get().BasicRegistry.FindOrAddHandle(SessionId);
 }
@@ -1393,7 +1393,7 @@ TOnlineAsyncOpHandle<FFindSessions> FSessionsEOSGS::FindSessions(FFindSessions::
 
 				const FFindSessions::Params& OpParams = StrongOp->GetParams();
 
-				TArray<FOnlineSessionIdHandle>& SearchResults = SearchResultsUserMap.FindChecked(OpParams.LocalAccountId);
+				TArray<FOnlineSessionId>& SearchResults = SearchResultsUserMap.FindChecked(OpParams.LocalAccountId);
 
 				StrongOp->SetResult({ SearchResults });
 
diff --git a/Engine/Plugins/Online/Experimental/OnlineServicesEOSGS/Source/Public/Online/NboSerializerEOSGSSvc.h b/Engine/Plugins/Online/Experimental/OnlineServicesEOSGS/Source/Public/Online/NboSerializerEOSGSSvc.h
--- a/Engine/Plugins/Online/Experimental/OnlineServicesEOSGS/Source/Public/Online/NboSerializerEOSGSSvc.h
+++ b/Engine/Plugins/Online/Experimental/OnlineServicesEOSGS/Source/Public/Online/NboSerializerEOSGSSvc.h
@@ -26,7 +26,7 @@ inline void SerializeToBuffer(FNboSerializeToBuffer& Ar, const FAccountId& Uniqu
 	Ar.WriteBinary(Data.GetData(), Data.Num());
 }
 
-inline void SerializeToBuffer(FNboSerializeToBuffer& Ar, const FOnlineSessionIdHandle& SessionId)
+inline void SerializeToBuffer(FNboSerializeToBuffer& Ar, const FOnlineSessionId& SessionId)
 {
 	TArray<uint8> Data = FOnlineSessionIdRegistryEOSGS::Get().ToReplicationData(SessionId);
 	Ar << Data.Num();
@@ -66,7 +66,7 @@ inline void SerializeFromBuffer(FNboSerializeFromBuffer& Ar, FAccountId& UniqueI
 	UniqueId = FOnlineAccountIdRegistryEOSGS::GetRegistered().FromReplicationData(Data);
 }
 
-inline void SerializeFromBuffer(FNboSerializeFromBuffer& Ar, FOnlineSessionIdHandle& SessionId)
+inline void SerializeFromBuffer(FNboSerializeFromBuffer& Ar, FOnlineSessionId& SessionId)
 {
 	TArray<uint8> Data;
 	int32 Size;
diff --git a/Engine/Plugins/Online/Experimental/OnlineServicesEOSGS/Source/Public/Online/SessionsEOSGS.h b/Engine/Plugins/Online/Experimental/OnlineServicesEOSGS/Source/Public/Online/SessionsEOSGS.h
--- a/Engine/Plugins/Online/Experimental/OnlineServicesEOSGS/Source/Public/Online/SessionsEOSGS.h
+++ b/Engine/Plugins/Online/Experimental/OnlineServicesEOSGS/Source/Public/Online/SessionsEOSGS.h
@@ -217,7 +217,7 @@ protected:
 
 	void SetSessionSearchMaxResults(FSessionSearchHandleEOSGS& SessionSearchHandle, uint32 MaxResults);
 	void SetSessionSearchParameters(FSessionSearchHandleEOSGS& SessionSearchHandle, TArray<FFindSessionsSearchFilter> Filters);
-	void SetSessionSearchSessionId(FSessionSearchHandleEOSGS& SessionSearchHandle, const FOnlineSessionIdHandle& SessionId);
+	void SetSessionSearchSessionId(FSessionSearchHandleEOSGS& SessionSearchHandle, const FOnlineSessionId& SessionId);
 	void SetSessionSearchTargetId(FSessionSearchHandleEOSGS& SessionSearchHandle, const FAccountId& TargetAccountId);
 
 	/**
@@ -250,7 +250,7 @@ protected:
 	 */
 	TOnlineAsyncOpHandle<FSendSingleSessionInviteImpl> SendSingleSessionInviteImpl(FSendSingleSessionInviteImpl::Params&& Params);
 
-	static FOnlineSessionIdHandle CreateSessionId(const FString& SessionId);
+	static FOnlineSessionId CreateSessionId(const FString& SessionId);
 	FOnlineSessionInviteIdHandle CreateSessionInviteId(const FString& SessionInviteId) const;
 
 	/**
diff --git a/Engine/Plugins/Online/Experimental/OnlineServicesNull/Source/Public/Online/NboSerializerNullSvc.h b/Engine/Plugins/Online/Experimental/OnlineServicesNull/Source/Public/Online/NboSerializerNullSvc.h
--- a/Engine/Plugins/Online/Experimental/OnlineServicesNull/Source/Public/Online/NboSerializerNullSvc.h
+++ b/Engine/Plugins/Online/Experimental/OnlineServicesNull/Source/Public/Online/NboSerializerNullSvc.h
@@ -26,7 +26,7 @@ inline void SerializeToBuffer(FNboSerializeToBuffer& Ar, const FAccountId& Uniqu
 	Ar.WriteBinary(Data.GetData(), Data.Num());
 }
 
-inline void SerializeToBuffer(FNboSerializeToBuffer& Ar, const FOnlineSessionIdHandle& SessionId)
+inline void SerializeToBuffer(FNboSerializeToBuffer& Ar, const FOnlineSessionId& SessionId)
 {
 	TArray<uint8> Data = FOnlineSessionIdRegistryNull::Get().ToReplicationData(SessionId);
 	Ar << Data.Num();
@@ -77,7 +77,7 @@ inline void SerializeFromBuffer(FNboSerializeFromBuffer& Ar, FAccountId& UniqueI
 	UniqueId = FOnlineAccountIdRegistryNull::Get().FromReplicationData(Data);
 }
 
-inline void SerializeFromBuffer(FNboSerializeFromBuffer& Ar, FOnlineSessionIdHandle& SessionId)
+inline void SerializeFromBuffer(FNboSerializeFromBuffer& Ar, FOnlineSessionId& SessionId)
 {
 	TArray<uint8> Data;
 	int32 Size;
diff --git a/Engine/Plugins/Online/Experimental/OnlineServicesOSSAdapter/Source/Private/Online/SessionsOSSAdapter.cpp b/Engine/Plugins/Online/Experimental/OnlineServicesOSSAdapter/Source/Private/Online/SessionsOSSAdapter.cpp
--- a/Engine/Plugins/Online/Experimental/OnlineServicesOSSAdapter/Source/Private/Online/SessionsOSSAdapter.cpp
+++ b/Engine/Plugins/Online/Experimental/OnlineServicesOSSAdapter/Source/Private/Online/SessionsOSSAdapter.cpp
@@ -306,7 +306,7 @@ void FSessionsOSSAdapter::Initialize()
 
 		FUISessionJoinRequested Event {
 			ServicesOSSAdapter.GetAccountIdRegistry().FindOrAddHandle(LocalAccountId),
-			TResult<FOnlineSessionIdHandle, FOnlineError>(GetSessionIdRegistry().FindOrAddHandle(InviteResult.Session.SessionInfo->GetSessionId().AsShared())),
+			TResult<FOnlineSessionId, FOnlineError>(GetSessionIdRegistry().FindOrAddHandle(InviteResult.Session.SessionInfo->GetSessionId().AsShared())),
 			EUISessionJoinRequestedSource::FromInvitation
 		};
 
@@ -654,7 +654,7 @@ TOnlineAsyncOpHandle<FFindSessions> FSessionsOSSAdapter::FindSessions(FFindSessi
 								AddSearchResult(FoundSession, Op->GetParams().LocalAccountId);
 							}
 
-							TArray<FOnlineSessionIdHandle> SearchResults = SearchResultsUserMap.FindOrAdd(Op->GetParams().LocalAccountId);
+							TArray<FOnlineSessionId> SearchResults = SearchResultsUserMap.FindOrAdd(Op->GetParams().LocalAccountId);
 							Op->SetResult({ SearchResults });
 						}
 						else
@@ -715,7 +715,7 @@ TOnlineAsyncOpHandle<FFindSessions> FSessionsOSSAdapter::FindSessions(FFindSessi
 							AddSearchResult(FoundSession, Op->GetParams().LocalAccountId);
 						}
 
-						TArray<FOnlineSessionIdHandle> SearchResults = SearchResultsUserMap.FindOrAdd(Op->GetParams().LocalAccountId);
+						TArray<FOnlineSessionId> SearchResults = SearchResultsUserMap.FindOrAdd(Op->GetParams().LocalAccountId);
 						Op->SetResult({ SearchResults });
 					}
 					else
@@ -753,7 +753,7 @@ TOnlineAsyncOpHandle<FFindSessions> FSessionsOSSAdapter::FindSessions(FFindSessi
 							AddSearchResult(FoundSession, Op->GetParams().LocalAccountId);
 						}
 
-						TArray<FOnlineSessionIdHandle> SearchResults = SearchResultsUserMap.FindOrAdd(Op->GetParams().LocalAccountId);
+						TArray<FOnlineSessionId> SearchResults = SearchResultsUserMap.FindOrAdd(Op->GetParams().LocalAccountId);
 						Op->SetResult({ SearchResults });
 					}
 					else
diff --git a/Engine/Plugins/Online/Experimental/OnlineServicesOSSAdapter/Source/Public/Online/OnlineIdOSSAdapter.h b/Engine/Plugins/Online/Experimental/OnlineServicesOSSAdapter/Source/Public/Online/OnlineIdOSSAdapter.h
--- a/Engine/Plugins/Online/Experimental/OnlineServicesOSSAdapter/Source/Public/Online/OnlineIdOSSAdapter.h
+++ b/Engine/Plugins/Online/Experimental/OnlineServicesOSSAdapter/Source/Public/Online/OnlineIdOSSAdapter.h
@@ -78,9 +78,9 @@ public:
 		return GetIdValue(Handle).ToSharedRef();
 	}
 
-	bool IsHandleExpired(const FOnlineSessionIdHandle& InHandle) const
+	bool IsHandleExpired(const FOnlineSessionId& InSessionId) const
 	{
-		return GetIdValue(InHandle).IsValid();
+		return GetIdValue(InSessionId).IsValid();
 	}
 
 	// Begin IOnlineAccountIdRegistry
diff --git a/Engine/Plugins/Online/Experimental/OnlineServicesOSSAdapter/Source/Public/Online/SessionsOSSAdapter.h b/Engine/Plugins/Online/Experimental/OnlineServicesOSSAdapter/Source/Public/Online/SessionsOSSAdapter.h
--- a/Engine/Plugins/Online/Experimental/OnlineServicesOSSAdapter/Source/Public/Online/SessionsOSSAdapter.h
+++ b/Engine/Plugins/Online/Experimental/OnlineServicesOSSAdapter/Source/Public/Online/SessionsOSSAdapter.h
@@ -16,7 +16,7 @@ namespace UE::Online {
 
 class FOnlineServicesOSSAdapter;
 
-/** Registry of V1 session ids in FUniqueNetIdRef format, indexed by FOnlineSessionIdHandle */
+/** Registry of V1 session ids in FUniqueNetIdRef format, indexed by FOnlineSessionId */
 using FOnlineSessionIdRegistryOSSAdapter = TOnlineUniqueNetIdRegistry<OnlineIdHandleTags::FSession>;
 
 using FOnlineSessionInviteIdRegistryOSSAdapter = FOnlineSessionInviteIdStringRegistry;
diff --git a/Engine/Source/Runtime/CoreOnline/Private/Online/CoreOnline.cpp b/Engine/Source/Runtime/CoreOnline/Private/Online/CoreOnline.cpp
--- a/Engine/Source/Runtime/CoreOnline/Private/Online/CoreOnline.cpp
+++ b/Engine/Source/Runtime/CoreOnline/Private/Online/CoreOnline.cpp
@@ -307,34 +307,34 @@ void FOnlineIdRegistryRegistry::UnregisterSessionIdRegistry(EOnlineServices Onli
 	}
 }
 
-FString FOnlineIdRegistryRegistry::ToLogString(const FOnlineSessionIdHandle& Handle) const
+FString FOnlineIdRegistryRegistry::ToLogString(const FOnlineSessionId& SessionId) const
 {
 	FString Result;
-	if (IOnlineSessionIdRegistry* Registry = GetSessionIdRegistry(Handle.GetOnlineServicesType()))
+	if (IOnlineSessionIdRegistry* Registry = GetSessionIdRegistry(SessionId.GetOnlineServicesType()))
 	{
-		Result = Registry->ToLogString(Handle);
+		Result = Registry->ToLogString(SessionId);
 	}
 	return Result;
 }
 
-TArray<uint8> FOnlineIdRegistryRegistry::ToReplicationData(const FOnlineSessionIdHandle& Handle) const
+TArray<uint8> FOnlineIdRegistryRegistry::ToReplicationData(const FOnlineSessionId& SessionId) const
 {
 	TArray<uint8> RepData;
-	if (IOnlineSessionIdRegistry* Registry = GetSessionIdRegistry(Handle.GetOnlineServicesType()))
+	if (IOnlineSessionIdRegistry* Registry = GetSessionIdRegistry(SessionId.GetOnlineServicesType()))
 	{
-		RepData = Registry->ToReplicationData(Handle);
+		RepData = Registry->ToReplicationData(SessionId);
 	}
 	return RepData;
 }
 
-FOnlineSessionIdHandle FOnlineIdRegistryRegistry::ToSessionId(EOnlineServices Services, const TArray<uint8>& RepData) const
+FOnlineSessionId FOnlineIdRegistryRegistry::ToSessionId(EOnlineServices Services, const TArray<uint8>& RepData) const
 {
-	FOnlineSessionIdHandle Handle;
+	FOnlineSessionId SessionId;
 	if (IOnlineSessionIdRegistry* Registry = GetSessionIdRegistry(Services))
 	{
-		Handle = Registry->FromReplicationData(RepData);
+		SessionId = Registry->FromReplicationData(RepData);
 	}
-	return Handle;
+	return SessionId;
 }
 
 IOnlineSessionIdRegistry* FOnlineIdRegistryRegistry::GetSessionIdRegistry(EOnlineServices OnlineServices) const
@@ -416,7 +416,7 @@ FString ToLogString(const FLobbyId& Id)
 	return FString::Printf(TEXT("%s:%d"), LexToString(Id.GetOnlineServicesType()), Id.GetHandle());
 }
 
-FString ToLogString(const FOnlineSessionIdHandle& Id)
+FString ToLogString(const FOnlineSessionId& Id)
 {
 	return FString::Printf(TEXT("%s:%d (%s)"), LexToString(Id.GetOnlineServicesType()), Id.GetHandle(), *FOnlineIdRegistryRegistry::Get().ToLogString(Id));
 }
diff --git a/Engine/Source/Runtime/CoreOnline/Public/Online/CoreOnline.h b/Engine/Source/Runtime/CoreOnline/Public/Online/CoreOnline.h
--- a/Engine/Source/Runtime/CoreOnline/Public/Online/CoreOnline.h
+++ b/Engine/Source/Runtime/CoreOnline/Public/Online/CoreOnline.h
@@ -308,14 +308,14 @@ private:
 
 using FAccountId = TOnlineIdHandle<OnlineIdHandleTags::FAccount>;
 using FLobbyId = TOnlineIdHandle<OnlineIdHandleTags::FLobby>;
-using FOnlineSessionIdHandle = TOnlineIdHandle<OnlineIdHandleTags::FSession>;
+using FOnlineSessionId = TOnlineIdHandle<OnlineIdHandleTags::FSession>;
 using FOnlineSessionInviteIdHandle = TOnlineIdHandle<OnlineIdHandleTags::FSessionInvite>;
 using FOnlineVerifiedAuthTicketIdHandle = TOnlineIdHandle<OnlineIdHandleTags::FVerifiedAuthTicket>;
 using FOnlineVerifiedAuthSessionIdHandle = TOnlineIdHandle<OnlineIdHandleTags::FVerifiedAuthSession>;
 
 COREONLINE_API FString ToLogString(const FAccountId& Id);
 COREONLINE_API FString ToLogString(const FLobbyId& Id);
-COREONLINE_API FString ToLogString(const FOnlineSessionIdHandle& Id);
+COREONLINE_API FString ToLogString(const FOnlineSessionId& Id);
 COREONLINE_API FString ToLogString(const FOnlineSessionInviteIdHandle& Id);
 COREONLINE_API FString ToLogString(const FOnlineVerifiedAuthTicketIdHandle& Id);
 COREONLINE_API FString ToLogString(const FOnlineVerifiedAuthSessionIdHandle& Id);
@@ -397,9 +397,9 @@ public:
 	 */
 	COREONLINE_API void UnregisterSessionIdRegistry(EOnlineServices OnlineServices, int32 Priority = 0);
 
-	COREONLINE_API FString ToLogString(const FOnlineSessionIdHandle& Handle) const;
-	COREONLINE_API TArray<uint8> ToReplicationData(const FOnlineSessionIdHandle& Handle) const;
-	COREONLINE_API FOnlineSessionIdHandle ToSessionId(EOnlineServices Services, const TArray<uint8>& RepData) const;
+	COREONLINE_API FString ToLogString(const FOnlineSessionId& SessionId) const;
+	COREONLINE_API TArray<uint8> ToReplicationData(const FOnlineSessionId& SessionId) const;
+	COREONLINE_API FOnlineSessionId ToSessionId(EOnlineServices Services, const TArray<uint8>& RepData) const;
 
 	COREONLINE_API IOnlineSessionIdRegistry* GetSessionIdRegistry(EOnlineServices OnlineServices) const;
 
