When canceling an online async task, make sure the outer error is "Errors::Cancelled"

diff --git a/Engine/Plugins/Online/Experimental/OnlineServices/Source/OnlineServicesCommon/Public/Online/OnlineAsyncOp.h b/Engine/Plugins/Online/Experimental/OnlineServices/Source/OnlineServicesCommon/Public/Online/OnlineAsyncOp.h
--- a/Engine/Plugins/Online/Experimental/OnlineServices/Source/OnlineServicesCommon/Public/Online/OnlineAsyncOp.h
+++ b/Engine/Plugins/Online/Experimental/OnlineServices/Source/OnlineServicesCommon/Public/Online/OnlineAsyncOp.h
@@ -873,7 +873,16 @@ protected:
 			if (PinnedOP.IsValid())
 			{
 				bCancelled = true;
-				TriggerOnComplete(TOnlineResult<OpType>(Reason));
+
+				// When canceling an operation the outer reason must always be Errors::Cancelled.
+				if (Reason.GetErrorCode() == Errors::ErrorCode::Common::Cancelled)
+				{
+					TriggerOnComplete(TOnlineResult<OpType>(Reason));
+				}
+				else
+				{
+					TriggerOnComplete(TOnlineResult<OpType>(Errors::Cancelled(Reason)));
+				}
 			}
 		}
 
