Removal of the concept of separate registered players from the Sessions interface

diff --git a/Engine/Plugins/Online/Experimental/OnlineServices/Source/OnlineServicesCommon/Private/Online/NboSerializerCommonSvc.cpp b/Engine/Plugins/Online/Experimental/OnlineServices/Source/OnlineServicesCommon/Private/Online/NboSerializerCommonSvc.cpp
--- a/Engine/Plugins/Online/Experimental/OnlineServices/Source/OnlineServicesCommon/Private/Online/NboSerializerCommonSvc.cpp
+++ b/Engine/Plugins/Online/Experimental/OnlineServices/Source/OnlineServicesCommon/Private/Online/NboSerializerCommonSvc.cpp
@@ -11,12 +11,6 @@ namespace NboSerializerCommonSvc {
 
 /** SerializeToBuffer methods */
 
-void SerializeToBuffer(FNboSerializeToBuffer& Packet, const FRegisteredPlayer& RegisteredPlayer)
-{
-	Packet << RegisteredPlayer.bHasReservedSlot;
-	Packet << RegisteredPlayer.bIsInSession;
-}
-
 void SerializeToBuffer(FNboSerializeToBuffer& Packet, const FSchemaVariant& Data)
 {
 	Packet << (uint8)Data.VariantType;
@@ -73,7 +67,6 @@ void SerializeToBuffer(FNboSerializeToBuffer& Packet, const FSessionSettings& Se
 {
 	Packet << SessionSettings.bAllowNewMembers;
 	Packet << SessionSettings.bAllowSanctionedPlayers;
-	Packet << SessionSettings.bAllowUnregisteredPlayers;
 	Packet << SessionSettings.bAntiCheatProtected;
 	Packet << SessionSettings.bIsDedicatedServerSession;
 	Packet << SessionSettings.bIsLANSession;
@@ -100,16 +93,6 @@ void SerializeToBuffer(FNboSerializeToBuffer& Packet, const FSessionCommon& Sess
 
 /** SerializeFromBuffer methods */
 
-void SerializeFromBuffer(FNboSerializeFromBuffer& Packet, FRegisteredPlayer& RegisteredPlayer)
-{
-	uint8 Read = 0;
-
-	Packet >> Read;
-	RegisteredPlayer.bHasReservedSlot = !!Read;
-	Packet >> Read;
-	RegisteredPlayer.bIsInSession = !!Read;
-}
-
 void SerializeFromBuffer(FNboSerializeFromBuffer& Packet, FSchemaVariant& Data)
 {
 	uint8 VariantType = 0;
@@ -190,8 +173,6 @@ void SerializeFromBuffer(FNboSerializeFromBuffer& Packet, FSessionSettings& Sess
 	Packet >> Read;
 	SessionSettings.bAllowSanctionedPlayers = !!Read;
 	Packet >> Read;
-	SessionSettings.bAllowUnregisteredPlayers = !!Read;
-	Packet >> Read;
 	SessionSettings.bAntiCheatProtected = !!Read;
 	Packet >> Read;
 	SessionSettings.bIsDedicatedServerSession = !!Read;
diff --git a/Engine/Plugins/Online/Experimental/OnlineServices/Source/OnlineServicesCommon/Private/Online/SessionsCommon.cpp b/Engine/Plugins/Online/Experimental/OnlineServices/Source/OnlineServicesCommon/Private/Online/SessionsCommon.cpp
--- a/Engine/Plugins/Online/Experimental/OnlineServices/Source/OnlineServicesCommon/Private/Online/SessionsCommon.cpp
+++ b/Engine/Plugins/Online/Experimental/OnlineServices/Source/OnlineServicesCommon/Private/Online/SessionsCommon.cpp
@@ -33,8 +33,6 @@ namespace UE::Online {
 		RegisterCommand(&FSessionsCommon::JoinSession);
 		RegisterCommand(&FSessionsCommon::SendSessionInvite);
 		RegisterCommand(&FSessionsCommon::RejectSessionInvite);
-		RegisterCommand(&FSessionsCommon::RegisterPlayers);
-		RegisterCommand(&FSessionsCommon::UnregisterPlayers);
 	}
 
 	TOnlineResult<FGetAllSessions> FSessionsCommon::GetAllSessions(FGetAllSessions::Params&& Params) const
@@ -247,68 +245,6 @@ namespace UE::Online {
 		return Operation->GetHandle();
 	}
 
-	// TODO: The logic for player registration vs adding a new session member is likely to be updated
-	TOnlineAsyncOpHandle<FRegisterPlayers> FSessionsCommon::RegisterPlayers(FRegisterPlayers::Params&& Params)
-	{
-		TOnlineAsyncOpRef<FRegisterPlayers> Op = GetOp<FRegisterPlayers>(MoveTemp(Params));
-
-		Op->Then([this](TOnlineAsyncOp<FRegisterPlayers>& Op) mutable
-		{
-			const FRegisterPlayers::Params& OpParams = Op.GetParams();
-
-			FOnlineError StateCheck = CheckRegisterPlayersState(OpParams);
-			if (StateCheck != Errors::Success())
-			{
-				Op.SetError(MoveTemp(StateCheck));
-				return;
-			}
-
-			TOnlineResult<FRegisterPlayers> Result = RegisterPlayersImpl(OpParams);
-			if (Result.IsOk())
-			{
-				Op.SetResult(MoveTemp(Result.GetOkValue()));
-			}
-			else
-			{
-				Op.SetError(MoveTemp(Result.GetErrorValue()));
-			}
-		})
-		.Enqueue(GetSerialQueue());
-
-		return Op->GetHandle();
-	}
-
-	// TODO: The logic for player registration removal vs removing a session member is likely to be updated
-	TOnlineAsyncOpHandle<FUnregisterPlayers> FSessionsCommon::UnregisterPlayers(FUnregisterPlayers::Params&& Params)
-	{
-		TOnlineAsyncOpRef<FUnregisterPlayers> Op = GetOp<FUnregisterPlayers>(MoveTemp(Params));
-
-		Op->Then([this](TOnlineAsyncOp<FUnregisterPlayers>& Op) mutable
-		{
-			const FUnregisterPlayers::Params& OpParams = Op.GetParams();
-
-			FOnlineError StateCheck = CheckUnregisterPlayersState(OpParams);
-			if (StateCheck != Errors::Success())
-			{
-				Op.SetError(MoveTemp(StateCheck));
-				return;
-			}
-
-			TOnlineResult<FUnregisterPlayers> Result = UnregisterPlayersImpl(OpParams);
-			if (Result.IsOk())
-			{
-				Op.SetResult(MoveTemp(Result.GetOkValue()));
-			}
-			else
-			{
-				Op.SetError(MoveTemp(Result.GetErrorValue()));
-			}
-		})
-		.Enqueue(GetSerialQueue());
-
-		return Op->GetHandle();
-	}
-
 	TOnlineEvent<void(const FSessionJoined&)> FSessionsCommon::OnSessionJoined()
 	{
 		return SessionEvents.OnSessionJoined;
@@ -487,58 +423,16 @@ namespace UE::Online {
 
 		FSessionSettings& SessionSettings = FoundSession->SessionSettings;
 
-		const FAccountId& SessionMemberId = Params.LocalAccountId;
-		const FSessionMember& SessionMember = Params.NewSessionMember;
-
-		// If the user was already registered in the session, we'll update it and reserve a slot for it
-		if (FRegisteredPlayer* RegisteredPlayer = SessionSettings.RegisteredPlayers.Find(SessionMemberId))
+		if (SessionSettings.NumOpenPublicConnections > 0)
 		{
-			RegisteredPlayer->bIsInSession = true;
-
-			if (!RegisteredPlayer->bHasReservedSlot)
-			{
-				bool bSlotReserved = false;
-
-				if (SessionSettings.NumOpenPublicConnections > 0)
-				{
-					SessionSettings.NumOpenPublicConnections--;
-					bSlotReserved = true;
-				}
-				else
-				{
-					return TOnlineResult<FAddSessionMember>(Errors::InvalidState());
-				}
-
-				// TODO: Discuss usage for Private Connections
-
-				RegisteredPlayer->bHasReservedSlot = bSlotReserved;
-			}
+			SessionSettings.NumOpenPublicConnections--;
 		}
-		else if (Params.bRegisterPlayers) // If it wasn't and the parameters dictate it, we'll register them with a reserved slot
-		{
-			FRegisterPlayers::Params RegisterPlayerParams = { Params.LocalAccountId, Params.SessionName, { SessionMemberId }, true };
-
-			TOnlineResult<FRegisterPlayers> Result = RegisterPlayersImpl(RegisterPlayerParams);
-			if (Result.IsError())
-			{
-				return TOnlineResult<FAddSessionMember>(MoveTemp(Result.GetErrorValue()));
-			}
-		}
-		else // Even if the new session member is not a registered player, we'll still need to decrease the number of open slots
+		else
 		{
-			if (SessionSettings.NumOpenPublicConnections > 0)
-			{
-				SessionSettings.NumOpenPublicConnections--;
-			}
-			else
-			{
-				return TOnlineResult<FAddSessionMember>(Errors::InvalidState());
-			}
-
-			// TODO: Discuss usage for Private Connections
+			return TOnlineResult<FAddSessionMember>(Errors::InvalidState());
 		}
 
-		SessionSettings.SessionMembers.Emplace(SessionMemberId, SessionMember);
+		SessionSettings.SessionMembers.Emplace(Params.LocalAccountId, Params.NewSessionMember);
 
 		return TOnlineResult<FAddSessionMember>();
 	}
@@ -554,32 +448,13 @@ namespace UE::Online {
 
 		FSessionSettings& SessionSettings = FoundSession->SessionSettings;
 		
-		if (Params.bUnregisterPlayers) // If the parameters dictate it, well unregister the player as well
+		if (SessionSettings.NumOpenPublicConnections < SessionSettings.NumMaxPublicConnections)
 		{
-			FUnregisterPlayers::Params UnregisterPlayersParams = { Params.LocalAccountId, Params.SessionName, { Params.LocalAccountId }, false };
-
-			TOnlineResult<FUnregisterPlayers> Result = UnregisterPlayersImpl(UnregisterPlayersParams);
-			if (Result.IsError())
-			{
-				return TOnlineResult<FRemoveSessionMember>(MoveTemp(Result.GetErrorValue()));
-			}
+			SessionSettings.NumOpenPublicConnections++;
 		}
-		else if (FRegisteredPlayer* RegisteredPlayer = SessionSettings.RegisteredPlayers.Find(Params.LocalAccountId)) // If the removed session member is still registered, we'll update it and keep their reserved slot
+		else
 		{
-			RegisteredPlayer->bIsInSession = false;
-		}
-		else  // If the removed session member is not a registered player with a reserved slot, we'll increase the number of open slots
-		{
-			if (SessionSettings.NumOpenPublicConnections < SessionSettings.NumMaxPublicConnections)
-			{
-				SessionSettings.NumOpenPublicConnections++;
-			}
-			else
-			{
-				return TOnlineResult<FRemoveSessionMember>(Errors::InvalidState());
-			}
-
-			// TODO: Discuss usage for Private Connections
+			return TOnlineResult<FRemoveSessionMember>(Errors::InvalidState());
 		}
 
 		SessionSettings.SessionMembers.Remove(Params.LocalAccountId);
@@ -587,104 +462,6 @@ namespace UE::Online {
 		return TOnlineResult<FRemoveSessionMember>();
 	}
 
-	TOnlineResult<FRegisterPlayers> FSessionsCommon::RegisterPlayersImpl(const FRegisterPlayers::Params& Params)
-	{
-		TOnlineResult<FGetMutableSessionByName> GetMutableSessionByNameResult = GetMutableSessionByName({ Params.SessionName });
-		if (GetMutableSessionByNameResult.IsError())
-		{
-			return TOnlineResult<FRegisterPlayers>(GetMutableSessionByNameResult.GetErrorValue());
-		}
-		TSharedRef<FSessionCommon> FoundSession = GetMutableSessionByNameResult.GetOkValue().Session;
-
-		FSessionSettings& SessionSettings = FoundSession->SessionSettings;
-
-		for (const FAccountId& User : Params.TargetUsers)
-		{
-			FRegisteredPlayer& RegisteredPlayer = SessionSettings.RegisteredPlayers.FindOrAdd(User);
-
-			if (const FSessionMember* SessionMember = SessionSettings.SessionMembers.Find(User))
-			{
-				RegisteredPlayer.bIsInSession = true;
-			}
-
-			// RegisterPlayers will not revoke a slot reserved for a registered player
-			if (!RegisteredPlayer.bIsInSession && !RegisteredPlayer.bHasReservedSlot && Params.bReserveSlot)
-			{
-				bool bSlotReserved = false;
-				
-				if (SessionSettings.NumOpenPublicConnections > 0)
-				{
-					SessionSettings.NumOpenPublicConnections--;
-					bSlotReserved = true;
-				}
-				else
-				{
-					return TOnlineResult<FRegisterPlayers>(Errors::InvalidState());
-				}
-
-				// TODO: Discuss usage for Private Connections
-
-				RegisteredPlayer.bHasReservedSlot = bSlotReserved;
-			}			
-		}
-
-		return TOnlineResult<FRegisterPlayers>();
-	}
-
-	TOnlineResult<FUnregisterPlayers> FSessionsCommon::UnregisterPlayersImpl(const FUnregisterPlayers::Params& Params)
-	{
-		TOnlineResult<FGetMutableSessionByName> GetMutableSessionByNameResult = GetMutableSessionByName({ Params.SessionName });
-		if (GetMutableSessionByNameResult.IsError())
-		{
-			return TOnlineResult<FUnregisterPlayers>(GetMutableSessionByNameResult.GetErrorValue());
-		}
-		TSharedRef<FSessionCommon> FoundSession = GetMutableSessionByNameResult.GetOkValue().Session;
-
-		FSessionSettings& SessionSettings = FoundSession->SessionSettings;
-
-		for (const FAccountId& User : Params.TargetUsers)
-		{
-			if (const FRegisteredPlayer* RegisteredPlayer = SessionSettings.RegisteredPlayers.Find(User))
-			{
-				// If the registered player had a reserved slot, we'll remove it
-				if (RegisteredPlayer->bHasReservedSlot)
-				{
-					if (SessionSettings.NumOpenPublicConnections < SessionSettings.NumMaxPublicConnections)
-					{
-						SessionSettings.NumOpenPublicConnections++;
-					}
-					else
-					{
-						return TOnlineResult<FUnregisterPlayers>(Errors::InvalidState());
-					}
-
-					// TODO: Discuss usage for Private Connections
-				}
-
-				if (RegisteredPlayer->bIsInSession && Params.bRemoveUnregisteredPlayers)
-				{
-					FRemoveSessionMember::Params RemoveSessionMemberParams = { User, Params.SessionName, false };
-
-					TOnlineResult<FRemoveSessionMember> Result = RemoveSessionMemberImpl(RemoveSessionMemberParams);
-					if (Result.IsError())
-					{
-						return TOnlineResult<FUnregisterPlayers>(MoveTemp(Result.GetErrorValue()));
-					}
-				}
-
-				SessionSettings.RegisteredPlayers.Remove(User);
-			}
-			else
-			{
-				UE_LOG(LogTemp, Warning, TEXT("[FSessionsCommon::UnregisterPlayers] Unable to unregister player with id [%s]. Player not registered"), *ToLogString(User));
-
-				return TOnlineResult<FUnregisterPlayers>(Errors::NotFound());
-			}
-		}
-
-		return TOnlineResult<FUnregisterPlayers>();
-	}
-
 	FOnlineError FSessionsCommon::CheckCreateSessionParams(const FCreateSession::Params& Params)
 	{
 		if (Params.SessionName.IsNone())
@@ -920,16 +697,6 @@ namespace UE::Online {
 			return Errors::AccessDenied();
 		}
 
-		if (!SessionSettings.bAllowUnregisteredPlayers)
-		{
-			if (!SessionSettings.RegisteredPlayers.Contains(Params.LocalAccountId))
-			{
-				UE_LOG(LogTemp, Warning, TEXT("[FSessionsCommon::CheckJoinSessionState] Could not join session. bAllowUnregisteredPlayers is set to false and user [%s] is not registered"), *ToLogString(Params.LocalAccountId));
-
-				return Errors::AccessDenied();
-			}
-		}
-
 		if (SessionSettings.NumMaxPublicConnections > 0 && SessionSettings.NumOpenPublicConnections > 0)
 		{
 			UE_LOG(LogTemp, Warning, TEXT("[FSessionsCommon::CheckJoinSessionState] Could not join session. No open public connections"));
@@ -1039,32 +806,6 @@ namespace UE::Online {
 		return Errors::Success();
 	}
 
-	FOnlineError FSessionsCommon::CheckRegisterPlayersState(const FRegisterPlayers::Params& Params)
-	{
-		if (TOptional<FOnlineError> Result = CheckSessionExistsByName(Params.LocalAccountId, Params.SessionName))
-		{
-			UE_LOG(LogTemp, Warning, TEXT("[FSessionsCommon::CheckRegisterPlayersState] Could not register players in session with name [%s]. Session not found"), *Params.SessionName.ToString());
-
-			return Result.GetValue();
-		}
-
-		// TODO: Check if there are enough slots available if Params::bReserveSlot is true
-
-		return Errors::Success();
-	}
-
-	FOnlineError FSessionsCommon::CheckUnregisterPlayersState(const FUnregisterPlayers::Params& Params)
-	{
-		if(TOptional<FOnlineError> Result = CheckSessionExistsByName(Params.LocalAccountId, Params.SessionName))
-		{
-			UE_LOG(LogTemp, Warning, TEXT("[FSessionsCommon::CheckUnregisterPlayersState] Could not unregister players from session with name [%s]. Session not found"), *Params.SessionName.ToString());
-
-			return Result.GetValue();
-		}
-
-		return Errors::Success();
-	}
-
 	// TODO: Have all Check methods return TOptional too, change call sites, and write Macro for repeating code structure
 	TOptional<FOnlineError> FSessionsCommon::CheckSessionExistsByName(const FAccountId& LocalAccountId, const FName& SessionName)
 	{
diff --git a/Engine/Plugins/Online/Experimental/OnlineServices/Source/OnlineServicesCommon/Public/Online/NboSerializerCommonSvc.h b/Engine/Plugins/Online/Experimental/OnlineServices/Source/OnlineServicesCommon/Public/Online/NboSerializerCommonSvc.h
--- a/Engine/Plugins/Online/Experimental/OnlineServices/Source/OnlineServicesCommon/Public/Online/NboSerializerCommonSvc.h
+++ b/Engine/Plugins/Online/Experimental/OnlineServices/Source/OnlineServicesCommon/Public/Online/NboSerializerCommonSvc.h
@@ -15,7 +15,6 @@ namespace NboSerializerCommonSvc {
 
 /** SerializeToBuffer methods */
 
-ONLINESERVICESCOMMON_API void SerializeToBuffer(FNboSerializeToBuffer& Packet, const FRegisteredPlayer& RegisteredPlayer);
 ONLINESERVICESCOMMON_API void SerializeToBuffer(FNboSerializeToBuffer& Packet, const FSchemaVariant& Data);
 ONLINESERVICESCOMMON_API void SerializeToBuffer(FNboSerializeToBuffer& Packet, const FCustomSessionSetting& CustomSessionSetting);
 ONLINESERVICESCOMMON_API void SerializeToBuffer(FNboSerializeToBuffer& Packet, const FCustomSessionSettingsMap& CustomSessionSettingsMap);
@@ -25,7 +24,6 @@ ONLINESERVICESCOMMON_API void SerializeToBuffer(FNboSerializeToBuffer& Packet, c
 
 /** SerializeFromBuffer methods */
 
-ONLINESERVICESCOMMON_API void SerializeFromBuffer(FNboSerializeFromBuffer& Packet, FRegisteredPlayer& RegisteredPlayer);
 ONLINESERVICESCOMMON_API void SerializeFromBuffer(FNboSerializeFromBuffer& Packet, FSchemaVariant& Data);
 ONLINESERVICESCOMMON_API void SerializeFromBuffer(FNboSerializeFromBuffer& Packet, FCustomSessionSetting& CustomSessionSetting);
 ONLINESERVICESCOMMON_API void SerializeFromBuffer(FNboSerializeFromBuffer& Packet, FCustomSessionSettingsMap& CustomSessionSettingsMap);
diff --git a/Engine/Plugins/Online/Experimental/OnlineServices/Source/OnlineServicesCommon/Public/Online/SessionsCommon.h b/Engine/Plugins/Online/Experimental/OnlineServices/Source/OnlineServicesCommon/Public/Online/SessionsCommon.h
--- a/Engine/Plugins/Online/Experimental/OnlineServices/Source/OnlineServicesCommon/Public/Online/SessionsCommon.h
+++ b/Engine/Plugins/Online/Experimental/OnlineServices/Source/OnlineServicesCommon/Public/Online/SessionsCommon.h
@@ -215,8 +215,6 @@ public:
 	virtual TOnlineAsyncOpHandle<FSendSessionInvite> SendSessionInvite(FSendSessionInvite::Params&& Params) override;
 	virtual TOnlineResult<FGetSessionInvites> GetSessionInvites(FGetSessionInvites::Params&& Params) override;
 	virtual TOnlineAsyncOpHandle<FRejectSessionInvite> RejectSessionInvite(FRejectSessionInvite::Params&& Params) override;
-	virtual TOnlineAsyncOpHandle<FRegisterPlayers> RegisterPlayers(FRegisterPlayers::Params&& Params) override;
-	virtual TOnlineAsyncOpHandle<FUnregisterPlayers> UnregisterPlayers(FUnregisterPlayers::Params&& Params) override;
 
 	virtual TOnlineEvent<void(const FSessionJoined&)> OnSessionJoined() override;
 	virtual TOnlineEvent<void(const FSessionLeft&)> OnSessionLeft() override;
@@ -239,9 +237,6 @@ protected:
 	TOnlineResult<FAddSessionMember> AddSessionMemberImpl(const FAddSessionMember::Params& Params);
 	TOnlineResult<FRemoveSessionMember> RemoveSessionMemberImpl(const FRemoveSessionMember::Params& Params);
 
-	TOnlineResult<FRegisterPlayers> RegisterPlayersImpl(const FRegisterPlayers::Params& Params);
-	TOnlineResult<FUnregisterPlayers> UnregisterPlayersImpl(const FUnregisterPlayers::Params& Params);
-
 	FOnlineError CheckCreateSessionParams(const FCreateSession::Params& Params);
 	FOnlineError CheckCreateSessionState(const FCreateSession::Params& Params);
 
@@ -266,10 +261,6 @@ protected:
 
 	FOnlineError CheckRejectSessionInviteState(const FRejectSessionInvite::Params& Params);
 
-	FOnlineError CheckRegisterPlayersState(const FRegisterPlayers::Params& Params);
-
-	FOnlineError CheckUnregisterPlayersState(const FUnregisterPlayers::Params& Params);
-
 private:
 	TOptional<FOnlineError> CheckSessionExistsByName(const FAccountId& LocalAccountId, const FName& SessionName);
 
diff --git a/Engine/Plugins/Online/Experimental/OnlineServices/Source/OnlineServicesInterface/Private/Online/Sessions.cpp b/Engine/Plugins/Online/Experimental/OnlineServices/Source/OnlineServicesInterface/Private/Online/Sessions.cpp
--- a/Engine/Plugins/Online/Experimental/OnlineServices/Source/OnlineServicesInterface/Private/Online/Sessions.cpp
+++ b/Engine/Plugins/Online/Experimental/OnlineServices/Source/OnlineServicesInterface/Private/Online/Sessions.cpp
@@ -54,7 +54,6 @@ FSessionSettings& FSessionSettings::operator+=(const FSessionSettingsUpdate& Upd
 	COPY_TOPTIONAL_VALUE_IF_SET(bIsDedicatedServerSession)
 	COPY_TOPTIONAL_VALUE_IF_SET(bAllowNewMembers)
 	COPY_TOPTIONAL_VALUE_IF_SET(bAllowSanctionedPlayers)
-	COPY_TOPTIONAL_VALUE_IF_SET(bAllowUnregisteredPlayers)
 	COPY_TOPTIONAL_VALUE_IF_SET(bAntiCheatProtected)
 	COPY_TOPTIONAL_VALUE_IF_SET(bPresenceEnabled)
 		
@@ -65,13 +64,6 @@ FSessionSettings& FSessionSettings::operator+=(const FSessionSettingsUpdate& Upd
 
 	CustomSettings.Append(UpdatedSettings.UpdatedCustomSettings);
 
-	for (const FAccountId& Key : UpdatedSettings.RemovedRegisteredPlayers)
-	{
-		RegisteredPlayers.Remove(Key);
-	}
-
-	RegisteredPlayers.Append(UpdatedSettings.UpdatedRegisteredPlayers);
-
 	for (const FAccountId& Key : UpdatedSettings.RemovedSessionMembers)
 	{
 		SessionMembers.Remove(Key);
@@ -134,7 +126,6 @@ FSessionSettingsUpdate& FSessionSettingsUpdate::operator+=(FSessionSettingsUpdat
 	MOVE_TOPTIONAL_IF_SET(bIsDedicatedServerSession)
 	MOVE_TOPTIONAL_IF_SET(bAllowNewMembers)
 	MOVE_TOPTIONAL_IF_SET(bAllowSanctionedPlayers)
-	MOVE_TOPTIONAL_IF_SET(bAllowUnregisteredPlayers)
 	MOVE_TOPTIONAL_IF_SET(bAntiCheatProtected)
 	MOVE_TOPTIONAL_IF_SET(bPresenceEnabled)
 
@@ -168,20 +159,6 @@ FSessionSettingsUpdate& FSessionSettingsUpdate::operator+=(FSessionSettingsUpdat
 		RemovedSessionMembers.AddUnique(MoveTemp(Key));
 	}
 
-	for (const TPair<FAccountId, FRegisteredPlayer>& UpdatedRegisteredPlayer : UpdatedValue.UpdatedRegisteredPlayers)
-	{
-		// If an update adds a modification to a registered player that had previously been marked for removal, we'll keep the latest change
-		RemovedRegisteredPlayers.Remove(UpdatedRegisteredPlayer.Key);
-	}
-	UpdatedRegisteredPlayers.Append(MoveTemp(UpdatedValue.UpdatedRegisteredPlayers));
-
-	for (FAccountId& Key : UpdatedValue.RemovedRegisteredPlayers)
-	{
-		// If an update removes a registered player that had previously been modified, we'll keep the latest change
-		UpdatedRegisteredPlayers.Remove(Key);
-		RemovedRegisteredPlayers.AddUnique(MoveTemp(Key));
-	}
-
 	return *this;
 }
 
diff --git a/Engine/Plugins/Online/Experimental/OnlineServices/Source/OnlineServicesInterface/Public/Online/Sessions.h b/Engine/Plugins/Online/Experimental/OnlineServices/Source/OnlineServicesInterface/Public/Online/Sessions.h
--- a/Engine/Plugins/Online/Experimental/OnlineServices/Source/OnlineServicesInterface/Public/Online/Sessions.h
+++ b/Engine/Plugins/Online/Experimental/OnlineServices/Source/OnlineServicesInterface/Public/Online/Sessions.h
@@ -63,18 +63,6 @@ enum class ESessionJoinPolicy : uint8
 ONLINESERVICESINTERFACE_API const TCHAR* LexToString(ESessionJoinPolicy Value);
 ONLINESERVICESINTERFACE_API void LexFromString(ESessionJoinPolicy& Value, const TCHAR* InStr);
 
-/** A player registered with the session, whether they are in it or not */
-struct FRegisteredPlayer
-{
-	/* Whether a slot was reserved for this player upon registration */
-	bool bHasReservedSlot = false;
-
-	/* Whether the player is currently in the session */
-	bool bIsInSession = false;
-};
-
-using FRegisteredPlayersMap = TMap<FAccountId, FRegisteredPlayer>;
-
 struct FSessionSettingsUpdate
 {
 	/** Set with an updated value if the SchemaName field will be changed in the update operation */
@@ -97,8 +85,6 @@ struct FSessionSettingsUpdate
 	TOptional<bool> bAllowNewMembers;
 	/** Set with an updated value if the bAllowSanctionedPlayers field will be changed in the update operation */
 	TOptional<bool> bAllowSanctionedPlayers;
-	/** Set with an updated value if the bAllowUnregisteredPlayers field will be changed in the update operation */
-	TOptional<bool> bAllowUnregisteredPlayers;
 	/** Set with an updated value if the bAntiCheatProtected field will be changed in the update operation */
 	TOptional<bool> bAntiCheatProtected;
 	/** Set with an updated value if the bPresenceEnabled field will be changed in the update operation */
@@ -114,11 +100,6 @@ struct FSessionSettingsUpdate
 	/** Id handles for session members to be removed in the update operation*/
 	TArray<FAccountId> RemovedSessionMembers;
 
-	/** Updated values for registered players to change in the update operation*/
-	FRegisteredPlayersMap UpdatedRegisteredPlayers;
-	/** Id handles for registered players to be removed in the update operation*/
-	TArray<FAccountId> RemovedRegisteredPlayers;
-
 	FSessionSettingsUpdate& operator+=(FSessionSettingsUpdate&& UpdatedValue);
 };
 
@@ -151,15 +132,12 @@ struct ONLINESERVICESINTERFACE_API FSessionSettings
 	/* Whether the session is configured to run as a dedicated server. Only available in some platforms. False by default */
 	bool bIsDedicatedServerSession = false;
 
-	/* Whether players (registered or not) are accepted as new members in the session. Can vary depending on various factors, like the number of free slots available, or Join-In-Progress preferences when the session has started. True by default */
+	/* Whether players are accepted as new members in the session. Can vary depending on various factors, like the number of free slots available, or Join-In-Progress preferences when the session has started. True by default */
 	bool bAllowNewMembers = true;
 
 	/* Whether this session will allow sanctioned players to join it. True by default */
 	bool bAllowSanctionedPlayers = true;
 
-	/* Whether this session will allow unregistered players to join it. True by default */
-	bool bAllowUnregisteredPlayers = true;
-
 	/*Whether this is a secure session protected by anti-cheat services. False by default */
 	bool bAntiCheatProtected = false;
 
@@ -171,9 +149,6 @@ struct ONLINESERVICESINTERFACE_API FSessionSettings
  
  	/* Map of session member ids to their corresponding user-defined settings */
  	FSessionMembersMap SessionMembers;
- 
- 	/* Map of registered players for this session. Can only be altered via RegisterPlayers and UnregisterPlayers */
- 	FRegisteredPlayersMap RegisteredPlayers;
 
 	FSessionSettings& operator+=(const FSessionSettingsUpdate& UpdatedValue);
 };
@@ -468,9 +443,6 @@ struct FAddSessionMember
 
 		/** Information for the session member to be added to the session. Any player that joins the session becomes a new member in doing so */
 		FSessionMember NewSessionMember;
-
-		/** Whether or not the new session members should also be added to the list of registered players. True by default*/
-		bool bRegisterPlayers = true;
 	};
 
 	struct Result
@@ -490,9 +462,6 @@ struct FRemoveSessionMember
 
 		/* Local name for the session */
 		FName SessionName;
-
-		/** Whether or not the session members should also be removed from the list of registered players. True by default*/
-		bool bUnregisterPlayers = true;
 	};
 
 	struct Result
@@ -558,54 +527,6 @@ struct FRejectSessionInvite
 	};
 };
 
-struct FRegisterPlayers
-{
-	static constexpr TCHAR Name[] = TEXT("RegisterPlayers");
-
-	struct Params
-	{
-		/* The local user agent */
-		FAccountId LocalAccountId;
-
-		/* The local name for the session. */
-		FName SessionName;
-
-		/* Array of users which will be registered */
-		TArray<FAccountId> TargetUsers;
-
-		/* Whether a slot should be saved for the registered players */
-		bool bReserveSlot;
-	};
-
-	struct Result
-	{
-	};
-};
-
-struct FUnregisterPlayers
-{
-	static constexpr TCHAR Name[] = TEXT("UnregisterPlayers");
-
-	struct Params
-	{
-		/* The local user agent */
-		FAccountId LocalAccountId;
-
-		/* The local name for the session. */
-		FName SessionName;
-
-		/* Array of users which will be unregistered */
-		TArray<FAccountId> TargetUsers;
-
-		/* Whether unregistered players should be removed from the session, if they are in it */
-		bool bRemoveUnregisteredPlayers;
-	};
-
-	struct Result
-	{
-	};
-};
-
 /* Events */
 
 struct FSessionJoined
@@ -771,7 +692,6 @@ public:
 	 * Adds a set of new session members to the named session
 	 * Session member information passed will be saved in the session settings
 	 * Number of open slots in the session will decrease accordingly
-	 * If indicated, players will also be registered in the session
 	 * 
 	 * @params Parameters for the AddSessionMember call
 	 * @return
@@ -782,7 +702,6 @@ public:
 	 * Removes a set of session member from the named session
 	 * Session member information for them will be removed from session settings
 	 * Number of open slots in the session will increase accordingly
-	 * If indicated, players will also be unregistered from the session
 	 *
 	 * @params Parameters for the RemoveSessionMember call
 	 * @return
@@ -813,24 +732,6 @@ public:
 	 */
 	virtual TOnlineAsyncOpHandle<FRejectSessionInvite> RejectSessionInvite(FRejectSessionInvite::Params&& Params) = 0;
 
-	/**
-	 * Registers given players in the named session
-	 * If indicated, and if any are available, a slot in the session will be reserved for them
-	 *
-	 * @param Parameters for the RegisterPlayers call
-	 * @return
-	 */
-	virtual TOnlineAsyncOpHandle<FRegisterPlayers> RegisterPlayers(FRegisterPlayers::Params&& Params) = 0;
-
-	/**
-	 * Unregisters given players from the named session.
-	 * If indicated, and if they are members of it, players will also be removed from the session
-	 *
-	 * @param Parameters for the UnregisterPlayers call
-	 * @return
-	 */
-	virtual TOnlineAsyncOpHandle<FUnregisterPlayers> UnregisterPlayers(FUnregisterPlayers::Params&& Params) = 0;
-
 	/* Events */
 
 	/**
@@ -897,11 +798,6 @@ BEGIN_ONLINE_STRUCT_META(FSessionMemberUpdate)
 	ONLINE_STRUCT_FIELD(FSessionMemberUpdate, RemovedMemberSettings)
 END_ONLINE_STRUCT_META()
 
-BEGIN_ONLINE_STRUCT_META(FRegisteredPlayer)
-	ONLINE_STRUCT_FIELD(FRegisteredPlayer, bHasReservedSlot),
-	ONLINE_STRUCT_FIELD(FRegisteredPlayer, bIsInSession)
-END_ONLINE_STRUCT_META()
-
 BEGIN_ONLINE_STRUCT_META(FSessionSettings)
 	ONLINE_STRUCT_FIELD(FSessionSettings, SchemaName),
 	ONLINE_STRUCT_FIELD(FSessionSettings, NumMaxPublicConnections),
@@ -914,12 +810,10 @@ BEGIN_ONLINE_STRUCT_META(FSessionSettings)
 	ONLINE_STRUCT_FIELD(FSessionSettings, bIsDedicatedServerSession),
 	ONLINE_STRUCT_FIELD(FSessionSettings, bAllowNewMembers),
 	ONLINE_STRUCT_FIELD(FSessionSettings, bAllowSanctionedPlayers),
-	ONLINE_STRUCT_FIELD(FSessionSettings, bAllowUnregisteredPlayers),
 	ONLINE_STRUCT_FIELD(FSessionSettings, bAntiCheatProtected),
 	ONLINE_STRUCT_FIELD(FSessionSettings, bPresenceEnabled),
 	ONLINE_STRUCT_FIELD(FSessionSettings, CustomSettings),
-	ONLINE_STRUCT_FIELD(FSessionSettings, SessionMembers),
-	ONLINE_STRUCT_FIELD(FSessionSettings, RegisteredPlayers)
+	ONLINE_STRUCT_FIELD(FSessionSettings, SessionMembers)
 END_ONLINE_STRUCT_META()
 
 BEGIN_ONLINE_STRUCT_META(FSessionSettingsUpdate)
@@ -933,15 +827,12 @@ BEGIN_ONLINE_STRUCT_META(FSessionSettingsUpdate)
 	ONLINE_STRUCT_FIELD(FSessionSettingsUpdate, bIsDedicatedServerSession),
 	ONLINE_STRUCT_FIELD(FSessionSettingsUpdate, bAllowNewMembers),
 	ONLINE_STRUCT_FIELD(FSessionSettingsUpdate, bAllowSanctionedPlayers),
-	ONLINE_STRUCT_FIELD(FSessionSettingsUpdate, bAllowUnregisteredPlayers),
 	ONLINE_STRUCT_FIELD(FSessionSettingsUpdate, bAntiCheatProtected),
 	ONLINE_STRUCT_FIELD(FSessionSettingsUpdate, bPresenceEnabled),
 	ONLINE_STRUCT_FIELD(FSessionSettingsUpdate, UpdatedCustomSettings),
 	ONLINE_STRUCT_FIELD(FSessionSettingsUpdate, RemovedCustomSettings),
 	ONLINE_STRUCT_FIELD(FSessionSettingsUpdate, UpdatedSessionMembers),
-	ONLINE_STRUCT_FIELD(FSessionSettingsUpdate, RemovedSessionMembers),
-	ONLINE_STRUCT_FIELD(FSessionSettingsUpdate, UpdatedRegisteredPlayers),
-	ONLINE_STRUCT_FIELD(FSessionSettingsUpdate, RemovedRegisteredPlayers)
+	ONLINE_STRUCT_FIELD(FSessionSettingsUpdate, RemovedSessionMembers)
 END_ONLINE_STRUCT_META()
 
 BEGIN_ONLINE_STRUCT_META(FSessionInvite)
@@ -1064,8 +955,7 @@ END_ONLINE_STRUCT_META()
 BEGIN_ONLINE_STRUCT_META(FAddSessionMember::Params)
 	ONLINE_STRUCT_FIELD(FAddSessionMember::Params, LocalAccountId),
 	ONLINE_STRUCT_FIELD(FAddSessionMember::Params, SessionName),
-	ONLINE_STRUCT_FIELD(FAddSessionMember::Params, NewSessionMember),
-	ONLINE_STRUCT_FIELD(FAddSessionMember::Params, bRegisterPlayers)
+	ONLINE_STRUCT_FIELD(FAddSessionMember::Params, NewSessionMember)
 END_ONLINE_STRUCT_META()
 
 BEGIN_ONLINE_STRUCT_META(FAddSessionMember::Result)
@@ -1073,8 +963,7 @@ END_ONLINE_STRUCT_META()
 
 BEGIN_ONLINE_STRUCT_META(FRemoveSessionMember::Params)
 	ONLINE_STRUCT_FIELD(FRemoveSessionMember::Params, LocalAccountId),
-	ONLINE_STRUCT_FIELD(FRemoveSessionMember::Params, SessionName),
-	ONLINE_STRUCT_FIELD(FRemoveSessionMember::Params, bUnregisterPlayers)
+	ONLINE_STRUCT_FIELD(FRemoveSessionMember::Params, SessionName)
 END_ONLINE_STRUCT_META()
 
 BEGIN_ONLINE_STRUCT_META(FRemoveSessionMember::Result)
@@ -1105,26 +994,6 @@ END_ONLINE_STRUCT_META()
 BEGIN_ONLINE_STRUCT_META(FRejectSessionInvite::Result)
 END_ONLINE_STRUCT_META()
 
-BEGIN_ONLINE_STRUCT_META(FRegisterPlayers::Params)
-	ONLINE_STRUCT_FIELD(FRegisterPlayers::Params, LocalAccountId),
-	ONLINE_STRUCT_FIELD(FRegisterPlayers::Params, SessionName),
-	ONLINE_STRUCT_FIELD(FRegisterPlayers::Params, TargetUsers),
-	ONLINE_STRUCT_FIELD(FRegisterPlayers::Params, bReserveSlot)
-END_ONLINE_STRUCT_META()
-
-BEGIN_ONLINE_STRUCT_META(FRegisterPlayers::Result)
-END_ONLINE_STRUCT_META()
-
-BEGIN_ONLINE_STRUCT_META(FUnregisterPlayers::Params)
-	ONLINE_STRUCT_FIELD(FUnregisterPlayers::Params, LocalAccountId),
-	ONLINE_STRUCT_FIELD(FUnregisterPlayers::Params, SessionName),
-	ONLINE_STRUCT_FIELD(FUnregisterPlayers::Params, TargetUsers),
-	ONLINE_STRUCT_FIELD(FUnregisterPlayers::Params, bRemoveUnregisteredPlayers)
-END_ONLINE_STRUCT_META()
-
-BEGIN_ONLINE_STRUCT_META(FUnregisterPlayers::Result)
-END_ONLINE_STRUCT_META()
-
 /* Meta*/ }
 
 /* UE::Online */ }
\ No newline at end of file
diff --git a/Engine/Plugins/Online/Experimental/OnlineServicesEOSGS/Source/Private/Online/SessionsEOSGS.cpp b/Engine/Plugins/Online/Experimental/OnlineServicesEOSGS/Source/Private/Online/SessionsEOSGS.cpp
--- a/Engine/Plugins/Online/Experimental/OnlineServicesEOSGS/Source/Private/Online/SessionsEOSGS.cpp
+++ b/Engine/Plugins/Online/Experimental/OnlineServicesEOSGS/Source/Private/Online/SessionsEOSGS.cpp
@@ -84,62 +84,30 @@ FSessionEOSGS::FSessionEOSGS(const EOS_HSessionDetails& InSessionDetailsHandle)
 				// We parse a single attribute
 				FString Key(Attribute->Data->Key);
 
-				// If the Key contains the ':' character, it will be a user related entry, either a Session Member setting, or Registered Player information
+				// If the Key contains the ':' character, it will be a Session Member setting
 				if (Key.Contains(TEXT(":")))
 				{
 					TArray<FString> KeyComponents;
 					Key.ParseIntoArray(KeyComponents, TEXT(":"));
 
-					// If the first element is the Registered Players key, we will process the entry as Registered Player information
-					if (KeyComponents[0] == EOSGS_REGISTERED_PLAYERS.ToString())
-					{
-						// We retrieve the registered player id
-						const FString& PlayerIdStr = KeyComponents[1];
-						const EOS_ProductUserId ProductUserId = EOS_ProductUserId_FromString(TCHAR_TO_UTF8(*PlayerIdStr));
-						const FAccountId& AccountId = FindAccountId(ProductUserId);
-
-						// This constructor should only be used by BuildSessionFromDetailsHandle, after all user ids in the session have been resolved
-						check(AccountId.IsValid());
-
-						FRegisteredPlayer& RegisteredPlayer = SessionSettings.RegisteredPlayers.FindOrAdd(AccountId);
-						
-						FSessionAttributeConverter<ESessionAttributeConversionType::FromService> CustomSettingConverter(*Attribute->Data);
-						const TPair<FSchemaAttributeId, FSchemaVariant>& CustomSettingData = CustomSettingConverter.GetAttributeData();
-						bool bParameterValue = CustomSettingData.Value.GetBoolean();
-
-						// And set the value for the appropriate parameter
-						const FString& ParameterName = KeyComponents[2];
-
-						if (ParameterName == EOSGS_REGISTERED_PLAYER_IS_IN_SESSION.ToString())
-						{
-							RegisteredPlayer.bIsInSession = bParameterValue;
-						}
-						else if (ParameterName == EOSGS_REGISTERED_PLAYER_HAS_RESERVED_SLOT.ToString())
-						{
-							RegisteredPlayer.bHasReservedSlot = bParameterValue;
-						}
-					}
-					else // If not, it will be parsed as a Session Member setting
-					{
-						// We retrieve the member id
-						FString PlayerIdStr = KeyComponents[0];
-						EOS_ProductUserId ProductUserId = EOS_ProductUserId_FromString(TCHAR_TO_UTF8(*PlayerIdStr));
-						FAccountId AccountId = FindAccountId(ProductUserId);
+					// We retrieve the member id
+					FString PlayerIdStr = KeyComponents[0];
+					EOS_ProductUserId ProductUserId = EOS_ProductUserId_FromString(TCHAR_TO_UTF8(*PlayerIdStr));
+					FAccountId AccountId = FindAccountId(ProductUserId);
 
-						FSessionMember& SessionMember = SessionSettings.SessionMembers.FindOrAdd(AccountId);
+					FSessionMember& SessionMember = SessionSettings.SessionMembers.FindOrAdd(AccountId);
 
-						// And add the corresponding custom setting
-						FSchemaAttributeId AttributeId(KeyComponents[1]);
+					// And add the corresponding custom setting
+					FSchemaAttributeId AttributeId(KeyComponents[1]);
 
-						FSessionAttributeConverter<ESessionAttributeConversionType::FromService> CustomSettingConverter(*Attribute->Data);
-						TPair<FSchemaAttributeId, FSchemaVariant> CustomSettingData = CustomSettingConverter.GetAttributeData();
+					FSessionAttributeConverter<ESessionAttributeConversionType::FromService> CustomSettingConverter(*Attribute->Data);
+					TPair<FSchemaAttributeId, FSchemaVariant> CustomSettingData = CustomSettingConverter.GetAttributeData();
 
-						FCustomSessionSetting CustomSessionSetting;
-						CustomSessionSetting.Visibility = FromServiceType(Attribute->AdvertisementType);
-						CustomSessionSetting.Data = CustomSettingData.Value;
+					FCustomSessionSetting CustomSessionSetting;
+					CustomSessionSetting.Visibility = FromServiceType(Attribute->AdvertisementType);
+					CustomSessionSetting.Data = CustomSettingData.Value;
 
-						SessionMember.MemberSettings.Emplace(AttributeId, CustomSessionSetting);
-					}
+					SessionMember.MemberSettings.Emplace(AttributeId, CustomSessionSetting);
 				}
 				else 
 				{
@@ -151,10 +119,6 @@ FSessionEOSGS::FSessionEOSGS(const EOS_HSessionDetails& InSessionDetailsHandle)
 					{
 						SessionSettings.bAllowNewMembers = CustomSettingData.Value.GetBoolean();
 					}
-					else if (Key == EOSGS_ALLOW_UNREGISTERED_PLAYERS.ToString())
-					{
-						SessionSettings.bAllowUnregisteredPlayers = CustomSettingData.Value.GetBoolean();
-					}
 					else if (Key == EOSGS_ANTI_CHEAT_PROTECTED.ToString())
 					{
 						SessionSettings.bAntiCheatProtected = CustomSettingData.Value.GetBoolean();
@@ -629,7 +593,6 @@ void FSessionsEOSGS::WriteCreateSessionModificationHandle(EOS_HSessionModificati
 
 	// We won't copy bIsLANSession since it' irrelevant for EOS Sessions
 	AddAttribute(SessionModificationHandle, EOSGS_ALLOW_NEW_MEMBERS, { FSchemaVariant(SessionSettings.bAllowNewMembers), ESchemaAttributeVisibility::Public });
-	AddAttribute(SessionModificationHandle, EOSGS_ALLOW_UNREGISTERED_PLAYERS, { FSchemaVariant(SessionSettings.bAllowUnregisteredPlayers), ESchemaAttributeVisibility::Public });
 	AddAttribute(SessionModificationHandle, EOSGS_ANTI_CHEAT_PROTECTED, { FSchemaVariant(SessionSettings.bAntiCheatProtected), ESchemaAttributeVisibility::Public });
 	AddAttribute(SessionModificationHandle, EOSGS_IS_DEDICATED_SERVER_SESSION, { FSchemaVariant(SessionSettings.bIsDedicatedServerSession), ESchemaAttributeVisibility::Public });
 	AddAttribute(SessionModificationHandle, EOSGS_PRESENCE_ENABLED, { FSchemaVariant(SessionSettings.bPresenceEnabled), ESchemaAttributeVisibility::Public });
@@ -667,17 +630,6 @@ void FSessionsEOSGS::WriteCreateSessionModificationHandle(EOS_HSessionModificati
 			AddAttribute(SessionModificationHandle, Key, CustomSettingEntry.Value);
 		}		
 	}
-
-	// TODO: Should always be empty at session creation time so might not be needed
-	// Registered Players
-	for (const TPair<FAccountId, FRegisteredPlayer>& Entry : SessionSettings.RegisteredPlayers)
-	{
-		const FName& HasReservedSlotKey = FName(EOSGS_REGISTERED_PLAYERS.ToString() + TEXT(":") + LexToString(GetProductUserIdChecked(Entry.Key)) + TEXT(":") + EOSGS_REGISTERED_PLAYER_HAS_RESERVED_SLOT.ToString());
-		AddAttribute(SessionModificationHandle, HasReservedSlotKey, { FSchemaVariant(Entry.Value.bHasReservedSlot), ESchemaAttributeVisibility::Public });
-
-		const FName& IsInSessionKey = FName(EOSGS_REGISTERED_PLAYERS.ToString() + TEXT(":") + LexToString(GetProductUserIdChecked(Entry.Key)) + TEXT(":") + EOSGS_REGISTERED_PLAYER_IS_IN_SESSION.ToString());
-		AddAttribute(SessionModificationHandle, IsInSessionKey, { FSchemaVariant(Entry.Value.bIsInSession), ESchemaAttributeVisibility::Public });
-	}
 }
 
 void FSessionsEOSGS::WriteUpdateSessionModificationHandle(EOS_HSessionModification& SessionModificationHandle, const FName& SessionName, const FSessionSettingsUpdate& NewSettings)
@@ -694,11 +646,6 @@ void FSessionsEOSGS::WriteUpdateSessionModificationHandle(EOS_HSessionModificati
 		SetInvitesAllowed(SessionModificationHandle, NewSettings.bAllowNewMembers.GetValue());
 	}
 
-	if (NewSettings.bAllowUnregisteredPlayers.IsSet())
-	{
-		AddAttribute(SessionModificationHandle, EOSGS_ALLOW_UNREGISTERED_PLAYERS, { FSchemaVariant(NewSettings.bAllowUnregisteredPlayers.GetValue()), ESchemaAttributeVisibility::Public });
-	}
-
 	if (NewSettings.bAntiCheatProtected.IsSet())
 	{
 		AddAttribute(SessionModificationHandle, EOSGS_ANTI_CHEAT_PROTECTED, { FSchemaVariant(NewSettings.bAntiCheatProtected.GetValue()), ESchemaAttributeVisibility::Public });
@@ -796,26 +743,6 @@ void FSessionsEOSGS::WriteUpdateSessionModificationHandle(EOS_HSessionModificati
 			AddAttribute(SessionModificationHandle, Key, CustomSettingEntry.Value);
 		}
 	}
-
-	// Registered Players
-	
-	for (const FAccountId& RemovedEntry : NewSettings.RemovedRegisteredPlayers)
-	{
-		const FSchemaAttributeId& HasReservedSlotKey = FSchemaAttributeId(EOSGS_REGISTERED_PLAYERS.ToString() + TEXT(":") + LexToString(GetProductUserIdChecked(RemovedEntry)) + TEXT(":") + EOSGS_REGISTERED_PLAYER_HAS_RESERVED_SLOT.ToString());
-		RemoveAttribute(SessionModificationHandle, HasReservedSlotKey);
-
-		const FSchemaAttributeId& IsInSessionKey = FSchemaAttributeId(EOSGS_REGISTERED_PLAYERS.ToString() + TEXT(":") + LexToString(GetProductUserIdChecked(RemovedEntry)) + TEXT(":") + EOSGS_REGISTERED_PLAYER_IS_IN_SESSION.ToString());
-		RemoveAttribute(SessionModificationHandle, IsInSessionKey);
-	}
-
-	for (const TPair<FAccountId, FRegisteredPlayer>& Entry : NewSettings.UpdatedRegisteredPlayers)
-	{
-		const FSchemaAttributeId& HasReservedSlotKey = FSchemaAttributeId(EOSGS_REGISTERED_PLAYERS.ToString() + TEXT(":") + LexToString(GetProductUserIdChecked(Entry.Key)) + TEXT(":") + EOSGS_REGISTERED_PLAYER_HAS_RESERVED_SLOT.ToString());
-		AddAttribute(SessionModificationHandle, HasReservedSlotKey, { FSchemaVariant(Entry.Value.bHasReservedSlot), ESchemaAttributeVisibility::Public });
-
-		const FSchemaAttributeId& IsInSessionKey = FSchemaAttributeId(EOSGS_REGISTERED_PLAYERS.ToString() + TEXT(":") + LexToString(GetProductUserIdChecked(Entry.Key)) + TEXT(":") + EOSGS_REGISTERED_PLAYER_IS_IN_SESSION.ToString());
-		AddAttribute(SessionModificationHandle, IsInSessionKey, { FSchemaVariant(Entry.Value.bIsInSession), ESchemaAttributeVisibility::Public });
-	}
 }
 
 TOnlineAsyncOpHandle<FUpdateSession> FSessionsEOSGS::UpdateSession(FUpdateSession::Params&& Params)
@@ -1537,7 +1464,7 @@ TOnlineAsyncOpHandle<FJoinSession> FSessionsEOSGS::JoinSession(FJoinSession::Par
 
 			SessionEvents.OnSessionJoined.Broadcast(Event);
 
-		// A successful join allows the client to server travel, after which RegisterPlayers will be called by the engine
+		// A successful join allows the client to server travel, after which AddSessionMember will be called by the engine
 	})
 	.Enqueue(GetSerialQueue());
 
@@ -1631,7 +1558,7 @@ TResult<TArray<EOS_ProductUserId>, FOnlineError> GetProductUserIdsFromEOSGSSessi
 				{
 					TArray<FString> KeyComponents;
 					Key.ParseIntoArray(KeyComponents, TEXT(":"));
-					const FString& PlayerIdStr = KeyComponents[0] == EOSGS_REGISTERED_PLAYERS.ToString() ? KeyComponents[1] : KeyComponents[0];
+					const FString& PlayerIdStr = KeyComponents[0];
 					const EOS_ProductUserId ProductUserId = EOS_ProductUserId_FromString(TCHAR_TO_UTF8(*PlayerIdStr));
 					Result.AddUnique(ProductUserId);
 				}
@@ -1932,7 +1859,6 @@ void FSessionsEOSGS::AppendSessionToPacket(FNboSerializeToBuffer& Packet, const
 	SerializeToBuffer(Packet, Session);
 	SerializeToBuffer(Packet, Session.OwnerAccountId);
 	SerializeToBuffer(Packet, Session.SessionId);
-	SerializeToBuffer(Packet, Session.SessionSettings.RegisteredPlayers);
 	SerializeToBuffer(Packet, Session.SessionSettings.SessionMembers);
 }
 
@@ -1944,7 +1870,6 @@ void FSessionsEOSGS::ReadSessionFromPacket(FNboSerializeFromBuffer& Packet, FSes
 	SerializeFromBuffer(Packet, Session);
 	SerializeFromBuffer(Packet, Session.OwnerAccountId);
 	SerializeFromBuffer(Packet, Session.SessionId);
-	SerializeFromBuffer(Packet, Session.SessionSettings.RegisteredPlayers);
 	SerializeFromBuffer(Packet, Session.SessionSettings.SessionMembers);
 }
 
diff --git a/Engine/Plugins/Online/Experimental/OnlineServicesEOSGS/Source/Public/Online/NboSerializerEOSGSSvc.h b/Engine/Plugins/Online/Experimental/OnlineServicesEOSGS/Source/Public/Online/NboSerializerEOSGSSvc.h
--- a/Engine/Plugins/Online/Experimental/OnlineServicesEOSGS/Source/Public/Online/NboSerializerEOSGSSvc.h
+++ b/Engine/Plugins/Online/Experimental/OnlineServicesEOSGS/Source/Public/Online/NboSerializerEOSGSSvc.h
@@ -44,17 +44,6 @@ inline void SerializeToBuffer(FNboSerializeToBuffer& Packet, const FSessionMembe
 	}
 }
 
-inline void SerializeToBuffer(FNboSerializeToBuffer& Packet, const FRegisteredPlayersMap& RegisteredPlayersMap)
-{
-	Packet << RegisteredPlayersMap.Num();
-
-	for (const TPair<FAccountId, FRegisteredPlayer>& Entry : RegisteredPlayersMap)
-	{
-		SerializeToBuffer(Packet, Entry.Key);
-		NboSerializerCommonSvc::SerializeToBuffer(Packet, Entry.Value);
-	}
-}
-
 /** NboSerializeFromBuffer methods */
 
 inline void SerializeFromBuffer(FNboSerializeFromBuffer& Ar, FAccountId& UniqueId)
@@ -92,23 +81,6 @@ inline void SerializeFromBuffer(FNboSerializeFromBuffer& Packet, FSessionMembers
 	}
 }
 
-inline void SerializeFromBuffer(FNboSerializeFromBuffer& Packet, FRegisteredPlayersMap& RegisteredPlayersMap)
-{
-	int32 NumEntries = 0;
-	Packet >> NumEntries;
-
-	for (int32 Index = 0; Index < NumEntries; ++Index)
-	{
-		FAccountId Key;
-		SerializeFromBuffer(Packet, Key);
-
-		FRegisteredPlayer Value;
-		NboSerializerCommonSvc::SerializeFromBuffer(Packet, Value);
-
-		RegisteredPlayersMap.Emplace(Key, Value);
-	}
-}
-
 /* NboSerializerNullSvc */ }
 
 /* UE::Online */ }
\ No newline at end of file
diff --git a/Engine/Plugins/Online/Experimental/OnlineServicesEOSGS/Source/Public/Online/SessionsEOSGS.h b/Engine/Plugins/Online/Experimental/OnlineServicesEOSGS/Source/Public/Online/SessionsEOSGS.h
--- a/Engine/Plugins/Online/Experimental/OnlineServicesEOSGS/Source/Public/Online/SessionsEOSGS.h
+++ b/Engine/Plugins/Online/Experimental/OnlineServicesEOSGS/Source/Public/Online/SessionsEOSGS.h
@@ -103,7 +103,7 @@ struct FUpdateSessionJoinabilityImpl
 		/** Name for the session, needed to start or end it */
 		FName SessionName;
 
-		/** Whether players (registered or not) are accepted as new members in the session. */
+		/** Whether players are accepted as new members in the session. */
 		bool bAllowNewMembers;
 	};
 
diff --git a/Engine/Plugins/Online/Experimental/OnlineServicesEOSGS/Source/Public/Online/SessionsEOSGSTypes.h b/Engine/Plugins/Online/Experimental/OnlineServicesEOSGS/Source/Public/Online/SessionsEOSGSTypes.h
--- a/Engine/Plugins/Online/Experimental/OnlineServicesEOSGS/Source/Public/Online/SessionsEOSGSTypes.h
+++ b/Engine/Plugins/Online/Experimental/OnlineServicesEOSGS/Source/Public/Online/SessionsEOSGSTypes.h
@@ -13,15 +13,11 @@
 namespace UE::Online {
 
 static FName EOSGS_ALLOW_NEW_MEMBERS = TEXT("EOSGS_ALLOW_NEW_MEMBERS");
-static FName EOSGS_ALLOW_UNREGISTERED_PLAYERS = TEXT("EOSGS_ALLOW_UNREGISTERED_PLAYERS");
 static FName EOSGS_ANTI_CHEAT_PROTECTED = TEXT("EOSGS_ANTI_CHEAT_PROTECTED");
 static FName EOSGS_IS_DEDICATED_SERVER_SESSION = TEXT("EOSGS_IS_DEDICATED_SERVER_SESSION");
 static FName EOSGS_PRESENCE_ENABLED = TEXT("EOSGS_PRESENCE_ENABLED");
 static FName EOSGS_SCHEMA_NAME = TEXT("EOSGS_SCHEMA_NAME");
 static FName EOSGS_SESSION_ID_OVERRIDE = TEXT("EOSGS_SESSION_ID_OVERRIDE");
-static FName EOSGS_REGISTERED_PLAYERS = TEXT("EOSGS_REGISTERED_PLAYERS");
-static FName EOSGS_REGISTERED_PLAYER_HAS_RESERVED_SLOT = TEXT("EOSGS_REGISTERED_PLAYER_HAS_RESERVED_SLOT");
-static FName EOSGS_REGISTERED_PLAYER_IS_IN_SESSION = TEXT("EOSGS_REGISTERED_PLAYER_IS_IN_SESSION");
 
 static FName EOSGS_BUCKET_ID = TEXT("EOSGS_BUCKET_ID");
 
diff --git a/Engine/Plugins/Online/Experimental/OnlineServicesNull/Source/Private/Online/SessionsNull.cpp b/Engine/Plugins/Online/Experimental/OnlineServicesNull/Source/Private/Online/SessionsNull.cpp
--- a/Engine/Plugins/Online/Experimental/OnlineServicesNull/Source/Private/Online/SessionsNull.cpp
+++ b/Engine/Plugins/Online/Experimental/OnlineServicesNull/Source/Private/Online/SessionsNull.cpp
@@ -38,7 +38,6 @@ void FSessionsNull::AppendSessionToPacket(FNboSerializeToBuffer& Packet, const F
 	SerializeToBuffer(Packet, Session);
 	SerializeToBuffer(Packet, Session.OwnerAccountId);
 	SerializeToBuffer(Packet, Session.SessionId);
-	SerializeToBuffer(Packet, Session.SessionSettings.RegisteredPlayers);
 	SerializeToBuffer(Packet, Session.SessionSettings.SessionMembers);
 }
 
@@ -50,7 +49,6 @@ void FSessionsNull::ReadSessionFromPacket(FNboSerializeFromBuffer& Packet, FSess
 	SerializeFromBuffer(Packet, Session);
 	SerializeFromBuffer(Packet, Session.OwnerAccountId);
 	SerializeFromBuffer(Packet, Session.SessionId);
-	SerializeFromBuffer(Packet, Session.SessionSettings.RegisteredPlayers);
 	SerializeFromBuffer(Packet, Session.SessionSettings.SessionMembers);
 }
 
diff --git a/Engine/Plugins/Online/Experimental/OnlineServicesNull/Source/Public/Online/NboSerializerNullSvc.h b/Engine/Plugins/Online/Experimental/OnlineServicesNull/Source/Public/Online/NboSerializerNullSvc.h
--- a/Engine/Plugins/Online/Experimental/OnlineServicesNull/Source/Public/Online/NboSerializerNullSvc.h
+++ b/Engine/Plugins/Online/Experimental/OnlineServicesNull/Source/Public/Online/NboSerializerNullSvc.h
@@ -44,17 +44,6 @@ inline void SerializeToBuffer(FNboSerializeToBuffer& Packet, const FSessionMembe
 	}
 }
 
-inline void SerializeToBuffer(FNboSerializeToBuffer& Packet, const FRegisteredPlayersMap& RegisteredPlayersMap)
-{
-	Packet << RegisteredPlayersMap.Num();
-
-	for (const TPair<FAccountId, FRegisteredPlayer>& Entry : RegisteredPlayersMap)
-	{
-		SerializeToBuffer(Packet, Entry.Key);
-		NboSerializerCommonSvc::SerializeToBuffer(Packet, Entry.Value);
-	}
-}
-
 inline void SerializeToBuffer(FNboSerializeToBuffer& Ar, const TMap<FLobbyAttributeId, FLobbyVariant>& Map)
 {
 	Ar << Map.Num();
@@ -104,23 +93,6 @@ inline void SerializeFromBuffer(FNboSerializeFromBuffer& Packet, FSessionMembers
 	}
 }
 
-inline void SerializeFromBuffer(FNboSerializeFromBuffer& Packet, FRegisteredPlayersMap& RegisteredPlayersMap)
-{
-	int32 NumEntries = 0;
-	Packet >> NumEntries;
-
-	for (int32 Index = 0; Index < NumEntries; ++Index)
-	{
-		FAccountId Key;
-		SerializeFromBuffer(Packet, Key);
-
-		FRegisteredPlayer Value;
-		NboSerializerCommonSvc::SerializeFromBuffer(Packet, Value);
-
-		RegisteredPlayersMap.Emplace(Key, Value);
-	}
-}
-
 inline void SerializeFromBuffer(FNboSerializeFromBuffer& Ar, TMap<FLobbyAttributeId, FLobbyVariant>& Map)
 {
 	int32 Size;
diff --git a/Engine/Plugins/Online/Experimental/OnlineServicesOSSAdapter/Source/Private/Online/SessionsOSSAdapter.cpp b/Engine/Plugins/Online/Experimental/OnlineServicesOSSAdapter/Source/Private/Online/SessionsOSSAdapter.cpp
--- a/Engine/Plugins/Online/Experimental/OnlineServicesOSSAdapter/Source/Private/Online/SessionsOSSAdapter.cpp
+++ b/Engine/Plugins/Online/Experimental/OnlineServicesOSSAdapter/Source/Private/Online/SessionsOSSAdapter.cpp
@@ -1117,20 +1117,6 @@ TOnlineAsyncOpHandle<FSendSessionInvite> FSessionsOSSAdapter::SendSessionInvite(
 	return Op->GetHandle();
 }
 
-TOnlineAsyncOpHandle<FRegisterPlayers> FSessionsOSSAdapter::RegisterPlayers(FRegisterPlayers::Params&& Params)
-{
-	TOnlineAsyncOpRef<FRegisterPlayers> Op = GetOp<FRegisterPlayers>(MoveTemp(Params));
-	Op->SetError(Errors::NotImplemented());
-	return Op->GetHandle();
-}
-
-TOnlineAsyncOpHandle<FUnregisterPlayers> FSessionsOSSAdapter::UnregisterPlayers(FUnregisterPlayers::Params&& Params)
-{
-	TOnlineAsyncOpRef<FUnregisterPlayers> Op = GetOp<FUnregisterPlayers>(MoveTemp(Params));
-	Op->SetError(Errors::NotImplemented());
-	return Op->GetHandle();
-}
-
 TOnlineResult<FGetResolvedConnectString> FSessionsOSSAdapter::GetResolvedConnectString(const FGetResolvedConnectString::Params& Params)
 {
 	if (Params.SessionId.IsValid())
@@ -1173,7 +1159,6 @@ FOnlineSessionSettings FSessionsOSSAdapter::BuildV1Settings(const FSessionSettin
 	Result.bAllowJoinViaPresence = InSessionSettings.bAllowNewMembers;
 	Result.bAllowJoinViaPresenceFriendsOnly = InSessionSettings.bAllowNewMembers;
 	Result.Settings.Add(OSS_ADAPTER_SESSIONS_ALLOW_SANCTIONED_PLAYERS, InSessionSettings.bAllowSanctionedPlayers);
-	Result.Settings.Add(OSS_ADAPTER_SESSIONS_ALLOW_UNREGISTERED_PLAYERS, InSessionSettings.bAllowUnregisteredPlayers);
 	Result.bAntiCheatProtected = InSessionSettings.bAntiCheatProtected;
 	Result.bIsDedicated = InSessionSettings.bIsDedicatedServerSession;
 	Result.bIsLANMatch = InSessionSettings.bIsLANSession;
@@ -1215,7 +1200,6 @@ FOnlineSessionSettings FSessionsOSSAdapter::BuildV1Settings(const FSessionSettin
 void FSessionsOSSAdapter::WriteV2SessionSettingsFromV1Session(const FOnlineSession* InSession, FSessionSettings& OutSettings) const
 {
 	InSession->SessionSettings.Settings.FindChecked(OSS_ADAPTER_SESSIONS_ALLOW_SANCTIONED_PLAYERS).Data.GetValue(OutSettings.bAllowSanctionedPlayers);
-	InSession->SessionSettings.Settings.FindChecked(OSS_ADAPTER_SESSIONS_ALLOW_UNREGISTERED_PLAYERS).Data.GetValue(OutSettings.bAllowUnregisteredPlayers);
 	OutSettings.bAntiCheatProtected = InSession->SessionSettings.bAntiCheatProtected;
 	OutSettings.bPresenceEnabled = InSession->SessionSettings.bUsesPresence;
 	OutSettings.bIsDedicatedServerSession = InSession->SessionSettings.bIsDedicated;
@@ -1258,9 +1242,10 @@ void FSessionsOSSAdapter::WriteV2SessionSettingsFromV1NamedSession(const FNamedO
 
 	FOnlineServicesOSSAdapter& ServicesOSSAdapter = static_cast<FOnlineServicesOSSAdapter&>(Services);
 
+	// We'll add all the Registered users as SessionMembers with empty data for now
 	for (const FUniqueNetIdRef& RegisteredPlayer : InSession->RegisteredPlayers)
 	{
-		OutSettings.RegisteredPlayers.Add(ServicesOSSAdapter.GetAccountIdRegistry().FindOrAddHandle(RegisteredPlayer));
+		OutSettings.SessionMembers.Emplace(ServicesOSSAdapter.GetAccountIdRegistry().FindOrAddHandle(RegisteredPlayer));
 	}
 }
 
@@ -1275,15 +1260,15 @@ TSharedRef<FSessionCommon> FSessionsOSSAdapter::BuildV2Session(const FOnlineSess
 
 	FOnlineServicesOSSAdapter& ServicesOSSAdapter = static_cast<FOnlineServicesOSSAdapter&>(Services);
 
-	Session->SessionId = GetSessionIdRegistry().FindOrAddHandle(InSession->SessionInfo->GetSessionId().AsShared());
-	WriteV2SessionSettingsFromV1Session(InSession, Session->SessionSettings);
-
 	if (const FNamedOnlineSession* NamedInSession = static_cast<const FNamedOnlineSession*>(InSession))
 	{
 		Session->OwnerAccountId = ServicesOSSAdapter.GetAccountIdRegistry().FindOrAddHandle(NamedInSession->LocalOwnerId->AsShared());
 		WriteV2SessionSettingsFromV1NamedSession(NamedInSession, Session->SessionSettings);
 	}
 
+	Session->SessionId = GetSessionIdRegistry().FindOrAddHandle(InSession->SessionInfo->GetSessionId().AsShared());
+	WriteV2SessionSettingsFromV1Session(InSession, Session->SessionSettings);
+
 	return Session;
 }
 
diff --git a/Engine/Plugins/Online/Experimental/OnlineServicesOSSAdapter/Source/Public/Online/SessionsOSSAdapter.h b/Engine/Plugins/Online/Experimental/OnlineServicesOSSAdapter/Source/Public/Online/SessionsOSSAdapter.h
--- a/Engine/Plugins/Online/Experimental/OnlineServicesOSSAdapter/Source/Public/Online/SessionsOSSAdapter.h
+++ b/Engine/Plugins/Online/Experimental/OnlineServicesOSSAdapter/Source/Public/Online/SessionsOSSAdapter.h
@@ -72,8 +72,6 @@ public:
 	virtual TOnlineAsyncOpHandle<FAddSessionMember> AddSessionMember(FAddSessionMember::Params&& Params) override;
 	virtual TOnlineAsyncOpHandle<FRemoveSessionMember> RemoveSessionMember(FRemoveSessionMember::Params&& Params) override;
 	virtual TOnlineAsyncOpHandle<FSendSessionInvite> SendSessionInvite(FSendSessionInvite::Params&& Params) override;
-	virtual TOnlineAsyncOpHandle<FRegisterPlayers> RegisterPlayers(FRegisterPlayers::Params&& Params) override;
-	virtual TOnlineAsyncOpHandle<FUnregisterPlayers> UnregisterPlayers(FUnregisterPlayers::Params&& Params) override;
 
 	TOnlineResult<FGetResolvedConnectString> GetResolvedConnectString(const FGetResolvedConnectString::Params& Params);
 
