OSSv2: Rename LoadConfig() to UpadateConfig().

diff --git a/Engine/Plugins/Online/Experimental/OnlineServices/Source/OnlineServicesCommon/Private/Online/AchievementsCommon.cpp b/Engine/Plugins/Online/Experimental/OnlineServices/Source/OnlineServicesCommon/Private/Online/AchievementsCommon.cpp
--- a/Engine/Plugins/Online/Experimental/OnlineServices/Source/OnlineServicesCommon/Private/Online/AchievementsCommon.cpp
+++ b/Engine/Plugins/Online/Experimental/OnlineServices/Source/OnlineServicesCommon/Private/Online/AchievementsCommon.cpp
@@ -33,9 +33,9 @@ void FAchievementsCommon::Shutdown()
 	TOnlineComponent<IAchievements>::Shutdown();
 }
 
-void FAchievementsCommon::LoadConfig()
+void FAchievementsCommon::UpdateConfig()
 {
-	TOnlineComponent<IAchievements>::LoadConfig();
+	TOnlineComponent<IAchievements>::UpdateConfig();
 
 	const TCHAR* ConfigSection = TEXT("OnlineServices.AchievementUnlockRules");
 	GConfig->GetBool(ConfigSection, TEXT("AchievementDef_IsTitleManaged"), bIsTitleManaged, GEngineIni);
diff --git a/Engine/Plugins/Online/Experimental/OnlineServices/Source/OnlineServicesCommon/Private/Online/LeaderboardsCommon.cpp b/Engine/Plugins/Online/Experimental/OnlineServices/Source/OnlineServicesCommon/Private/Online/LeaderboardsCommon.cpp
--- a/Engine/Plugins/Online/Experimental/OnlineServices/Source/OnlineServicesCommon/Private/Online/LeaderboardsCommon.cpp
+++ b/Engine/Plugins/Online/Experimental/OnlineServices/Source/OnlineServicesCommon/Private/Online/LeaderboardsCommon.cpp
@@ -78,9 +78,9 @@ void FLeaderboardsCommon::Initialize()
 	StatEventHandle = Services.Get<IStats>()->OnStatsUpdated().Add([this](const FStatsUpdated& StatsUpdated) { WriteLeaderboardsByStats(StatsUpdated); });
 }
 
-void FLeaderboardsCommon::LoadConfig()
+void FLeaderboardsCommon::UpdateConfig()
 {
-	Super::LoadConfig();
+	Super::UpdateConfig();
 
 	const TCHAR* ConfigSection = TEXT("OnlineServices.Leaderboards");
 	GConfig->GetBool(ConfigSection, TEXT("LeaderboardDef_IsTitleManaged"), bIsTitleManaged, GEngineIni);
diff --git a/Engine/Plugins/Online/Experimental/OnlineServices/Source/OnlineServicesCommon/Private/Online/OnlineServicesCommon.cpp b/Engine/Plugins/Online/Experimental/OnlineServices/Source/OnlineServicesCommon/Private/Online/OnlineServicesCommon.cpp
--- a/Engine/Plugins/Online/Experimental/OnlineServices/Source/OnlineServicesCommon/Private/Online/OnlineServicesCommon.cpp
+++ b/Engine/Plugins/Online/Experimental/OnlineServices/Source/OnlineServicesCommon/Private/Online/OnlineServicesCommon.cpp
@@ -140,9 +140,9 @@ void FOnlineServicesCommon::PostInitialize()
 	Components.Visit(&IOnlineComponent::PostInitialize);
 }
 
-void FOnlineServicesCommon::LoadConfig()
+void FOnlineServicesCommon::UpdateConfig()
 {
-	Components.Visit(&IOnlineComponent::LoadConfig);
+	Components.Visit(&IOnlineComponent::UpdateConfig);
 }
 
 bool FOnlineServicesCommon::Tick(float DeltaSeconds)
diff --git a/Engine/Plugins/Online/Experimental/OnlineServices/Source/OnlineServicesCommon/Private/Online/StatsCommon.cpp b/Engine/Plugins/Online/Experimental/OnlineServices/Source/OnlineServicesCommon/Private/Online/StatsCommon.cpp
--- a/Engine/Plugins/Online/Experimental/OnlineServices/Source/OnlineServicesCommon/Private/Online/StatsCommon.cpp
+++ b/Engine/Plugins/Online/Experimental/OnlineServices/Source/OnlineServicesCommon/Private/Online/StatsCommon.cpp
@@ -88,7 +88,7 @@ FStatsCommon::FStatsCommon(FOnlineServicesCommon& InServices)
 {
 }
 
-void FStatsCommon::LoadConfig()
+void FStatsCommon::UpdateConfig()
 {
 	const TCHAR* ConfigSection = TEXT("OnlineServices.Stats");
 
diff --git a/Engine/Plugins/Online/Experimental/OnlineServices/Source/OnlineServicesCommon/Public/Online/AchievementsCommon.h b/Engine/Plugins/Online/Experimental/OnlineServices/Source/OnlineServicesCommon/Public/Online/AchievementsCommon.h
--- a/Engine/Plugins/Online/Experimental/OnlineServices/Source/OnlineServicesCommon/Public/Online/AchievementsCommon.h
+++ b/Engine/Plugins/Online/Experimental/OnlineServices/Source/OnlineServicesCommon/Public/Online/AchievementsCommon.h
@@ -34,7 +34,7 @@ public:
 	// TOnlineComponent
 	virtual void Initialize() override;
 	virtual void Shutdown() override;
-	virtual void LoadConfig() override;
+	virtual void UpdateConfig() override;
 	virtual void RegisterCommands() override;
 
 	// IAchievements
diff --git a/Engine/Plugins/Online/Experimental/OnlineServices/Source/OnlineServicesCommon/Public/Online/IOnlineComponent.h b/Engine/Plugins/Online/Experimental/OnlineServices/Source/OnlineServicesCommon/Public/Online/IOnlineComponent.h
--- a/Engine/Plugins/Online/Experimental/OnlineServices/Source/OnlineServicesCommon/Public/Online/IOnlineComponent.h
+++ b/Engine/Plugins/Online/Experimental/OnlineServices/Source/OnlineServicesCommon/Public/Online/IOnlineComponent.h
@@ -16,7 +16,7 @@ public:
 	// Called after all components have been initialized
 	virtual void PostInitialize() = 0;
 	// Called whenever we need to reload data from config
-	virtual void LoadConfig() = 0;
+	virtual void UpdateConfig() = 0;
 	// Called every Tick
 	virtual void Tick(float DeltaSeconds) = 0;
 	// Called before any component has been shutdown
diff --git a/Engine/Plugins/Online/Experimental/OnlineServices/Source/OnlineServicesCommon/Public/Online/LeaderboardsCommon.h b/Engine/Plugins/Online/Experimental/OnlineServices/Source/OnlineServicesCommon/Public/Online/LeaderboardsCommon.h
--- a/Engine/Plugins/Online/Experimental/OnlineServices/Source/OnlineServicesCommon/Public/Online/LeaderboardsCommon.h
+++ b/Engine/Plugins/Online/Experimental/OnlineServices/Source/OnlineServicesCommon/Public/Online/LeaderboardsCommon.h
@@ -66,7 +66,7 @@ public:
 
 	// TOnlineComponent
 	virtual void Initialize() override;
-	virtual void LoadConfig() override;
+	virtual void UpdateConfig() override;
 	virtual void RegisterCommands() override;
 
 	// ILeaderboards
diff --git a/Engine/Plugins/Online/Experimental/OnlineServices/Source/OnlineServicesCommon/Public/Online/OnlineComponent.h b/Engine/Plugins/Online/Experimental/OnlineServices/Source/OnlineServicesCommon/Public/Online/OnlineComponent.h
--- a/Engine/Plugins/Online/Experimental/OnlineServices/Source/OnlineServicesCommon/Public/Online/OnlineComponent.h
+++ b/Engine/Plugins/Online/Experimental/OnlineServices/Source/OnlineServicesCommon/Public/Online/OnlineComponent.h
@@ -34,10 +34,10 @@ public:
 	{
 		Services.RegisterExecHandler(InterfaceName, MakeUnique<TOnlineComponentExecHandler<TOnlineComponent<ComponentType>>>(this));
 		RegisterCommands();
-		LoadConfig();
+		UpdateConfig();
 	}
 	virtual void PostInitialize() override {}
-	virtual void LoadConfig() override {}
+	virtual void UpdateConfig() override {}
 	virtual void Tick(float DeltaSeconds) override {}
 	virtual void PreShutdown() override {}
 	virtual void Shutdown() override {}
diff --git a/Engine/Plugins/Online/Experimental/OnlineServices/Source/OnlineServicesCommon/Public/Online/OnlineServicesCommon.h b/Engine/Plugins/Online/Experimental/OnlineServices/Source/OnlineServicesCommon/Public/Online/OnlineServicesCommon.h
--- a/Engine/Plugins/Online/Experimental/OnlineServices/Source/OnlineServicesCommon/Public/Online/OnlineServicesCommon.h
+++ b/Engine/Plugins/Online/Experimental/OnlineServices/Source/OnlineServicesCommon/Public/Online/OnlineServicesCommon.h
@@ -80,9 +80,9 @@ public:
 	virtual void PostInitialize();
 
 	/**
-	 * Calls LoadConfig on all the components
+	 * Calls UpdateConfig on all the components
 	 */
-	virtual void LoadConfig();
+	virtual void UpdateConfig();
 
 	/**
 	 * Calls Tick on all the components
diff --git a/Engine/Plugins/Online/Experimental/OnlineServices/Source/OnlineServicesCommon/Public/Online/StatsCommon.h b/Engine/Plugins/Online/Experimental/OnlineServices/Source/OnlineServicesCommon/Public/Online/StatsCommon.h
--- a/Engine/Plugins/Online/Experimental/OnlineServices/Source/OnlineServicesCommon/Public/Online/StatsCommon.h
+++ b/Engine/Plugins/Online/Experimental/OnlineServices/Source/OnlineServicesCommon/Public/Online/StatsCommon.h
@@ -71,7 +71,7 @@ public:
 	FStatsCommon(FOnlineServicesCommon& InServices);
 
 	// TOnlineComponent
-	virtual void LoadConfig() override;
+	virtual void UpdateConfig() override;
 	virtual void RegisterCommands() override;
 
 	// IStats
diff --git a/Engine/Plugins/Online/Experimental/OnlineServicesEOS/Source/Private/Online/ExternalUIEOS.cpp b/Engine/Plugins/Online/Experimental/OnlineServicesEOS/Source/Private/Online/ExternalUIEOS.cpp
--- a/Engine/Plugins/Online/Experimental/OnlineServicesEOS/Source/Private/Online/ExternalUIEOS.cpp
+++ b/Engine/Plugins/Online/Experimental/OnlineServicesEOS/Source/Private/Online/ExternalUIEOS.cpp
@@ -29,9 +29,9 @@ void FExternalUIEOS::Initialize()
 	check(UIHandle != nullptr);
 }
 
-void FExternalUIEOS::LoadConfig()
+void FExternalUIEOS::UpdateConfig()
 {
-	Super::LoadConfig();
+	Super::UpdateConfig();
 	LoadConfig(Config);
 }
 
diff --git a/Engine/Plugins/Online/Experimental/OnlineServicesEOS/Source/Public/Online/ExternalUIEOS.h b/Engine/Plugins/Online/Experimental/OnlineServicesEOS/Source/Public/Online/ExternalUIEOS.h
--- a/Engine/Plugins/Online/Experimental/OnlineServicesEOS/Source/Public/Online/ExternalUIEOS.h
+++ b/Engine/Plugins/Online/Experimental/OnlineServicesEOS/Source/Public/Online/ExternalUIEOS.h
@@ -27,7 +27,7 @@ public:
 
 	FExternalUIEOS(FOnlineServicesEOS& InOwningSubsystem);
 	virtual void Initialize() override;
-	virtual void LoadConfig() override;
+	virtual void UpdateConfig() override;
 	virtual void PreShutdown() override;
 	virtual TOnlineAsyncOpHandle<FExternalUIShowLoginUI> ShowLoginUI(FExternalUIShowLoginUI::Params&& Params) override;
 	virtual TOnlineAsyncOpHandle<FExternalUIShowFriendsUI> ShowFriendsUI(FExternalUIShowFriendsUI::Params&& Params) override;
diff --git a/Engine/Plugins/Online/Experimental/OnlineServicesEOSGS/Source/Private/Online/TitleFileEOSGS.cpp b/Engine/Plugins/Online/Experimental/OnlineServicesEOSGS/Source/Private/Online/TitleFileEOSGS.cpp
--- a/Engine/Plugins/Online/Experimental/OnlineServicesEOSGS/Source/Private/Online/TitleFileEOSGS.cpp
+++ b/Engine/Plugins/Online/Experimental/OnlineServicesEOSGS/Source/Private/Online/TitleFileEOSGS.cpp
@@ -27,9 +27,9 @@ void FTitleFileEOSGS::Initialize()
 	check(TitleStorageHandle);
 }
 
-void FTitleFileEOSGS::LoadConfig()
+void FTitleFileEOSGS::UpdateConfig()
 {
-	Super::LoadConfig();
+	Super::UpdateConfig();
 	Super::LoadConfig(Config);
 }
 
diff --git a/Engine/Plugins/Online/Experimental/OnlineServicesEOSGS/Source/Private/Online/UserFileEOSGS.cpp b/Engine/Plugins/Online/Experimental/OnlineServicesEOSGS/Source/Private/Online/UserFileEOSGS.cpp
--- a/Engine/Plugins/Online/Experimental/OnlineServicesEOSGS/Source/Private/Online/UserFileEOSGS.cpp
+++ b/Engine/Plugins/Online/Experimental/OnlineServicesEOSGS/Source/Private/Online/UserFileEOSGS.cpp
@@ -30,9 +30,9 @@ void FUserFileEOSGS::Initialize()
 	check(PlayerDataStorageHandle);
 }
 
-void FUserFileEOSGS::LoadConfig()
+void FUserFileEOSGS::UpdateConfig()
 {
-	Super::LoadConfig();
+	Super::UpdateConfig();
 	Super::LoadConfig(Config);
 }
 
diff --git a/Engine/Plugins/Online/Experimental/OnlineServicesEOSGS/Source/Public/Online/TitleFileEOSGS.h b/Engine/Plugins/Online/Experimental/OnlineServicesEOSGS/Source/Public/Online/TitleFileEOSGS.h
--- a/Engine/Plugins/Online/Experimental/OnlineServicesEOSGS/Source/Public/Online/TitleFileEOSGS.h
+++ b/Engine/Plugins/Online/Experimental/OnlineServicesEOSGS/Source/Public/Online/TitleFileEOSGS.h
@@ -39,7 +39,7 @@ public:
 
 	// IOnlineComponent
 	virtual void Initialize() override;
-	virtual void LoadConfig() override;
+	virtual void UpdateConfig() override;
 
 	// ITitleFile
 	virtual TOnlineAsyncOpHandle<FTitleFileEnumerateFiles> EnumerateFiles(FTitleFileEnumerateFiles::Params&& Params) override;
diff --git a/Engine/Plugins/Online/Experimental/OnlineServicesEOSGS/Source/Public/Online/UserFileEOSGS.h b/Engine/Plugins/Online/Experimental/OnlineServicesEOSGS/Source/Public/Online/UserFileEOSGS.h
--- a/Engine/Plugins/Online/Experimental/OnlineServicesEOSGS/Source/Public/Online/UserFileEOSGS.h
+++ b/Engine/Plugins/Online/Experimental/OnlineServicesEOSGS/Source/Public/Online/UserFileEOSGS.h
@@ -37,7 +37,7 @@ public:
 
 	// IOnlineComponent
 	virtual void Initialize() override;
-	virtual void LoadConfig() override;
+	virtual void UpdateConfig() override;
 
 	// IUserFile
 	virtual TOnlineAsyncOpHandle<FUserFileEnumerateFiles> EnumerateFiles(FUserFileEnumerateFiles::Params&& Params) override;
diff --git a/Engine/Plugins/Online/Experimental/OnlineServicesNull/Source/Private/Online/TitleFileNull.cpp b/Engine/Plugins/Online/Experimental/OnlineServicesNull/Source/Private/Online/TitleFileNull.cpp
--- a/Engine/Plugins/Online/Experimental/OnlineServicesNull/Source/Private/Online/TitleFileNull.cpp
+++ b/Engine/Plugins/Online/Experimental/OnlineServicesNull/Source/Private/Online/TitleFileNull.cpp
@@ -15,7 +15,7 @@ FTitleFileNull::FTitleFileNull(FOnlineServicesNull& InOwningSubsystem)
 {
 }
 
-void FTitleFileNull::LoadConfig()
+void FTitleFileNull::UpdateConfig()
 {
 	const TCHAR* ConfigSection = TEXT("OnlineServices.Null.TitleFile");
 
diff --git a/Engine/Plugins/Online/Experimental/OnlineServicesNull/Source/Private/Online/UserFileNull.cpp b/Engine/Plugins/Online/Experimental/OnlineServicesNull/Source/Private/Online/UserFileNull.cpp
--- a/Engine/Plugins/Online/Experimental/OnlineServicesNull/Source/Private/Online/UserFileNull.cpp
+++ b/Engine/Plugins/Online/Experimental/OnlineServicesNull/Source/Private/Online/UserFileNull.cpp
@@ -15,7 +15,7 @@ FUserFileNull::FUserFileNull(FOnlineServicesNull& InOwningSubsystem)
 {
 }
 
-void FUserFileNull::LoadConfig()
+void FUserFileNull::UpdateConfig()
 {
 	const TCHAR* ConfigSection = TEXT("OnlineServices.Null.UserFile");
 
diff --git a/Engine/Plugins/Online/Experimental/OnlineServicesNull/Source/Public/Online/TitleFileNull.h b/Engine/Plugins/Online/Experimental/OnlineServicesNull/Source/Public/Online/TitleFileNull.h
--- a/Engine/Plugins/Online/Experimental/OnlineServicesNull/Source/Public/Online/TitleFileNull.h
+++ b/Engine/Plugins/Online/Experimental/OnlineServicesNull/Source/Public/Online/TitleFileNull.h
@@ -17,7 +17,7 @@ public:
 	FTitleFileNull(FOnlineServicesNull& InOwningSubsystem);
 
 	// IOnlineComponent
-	virtual void LoadConfig() override;
+	virtual void UpdateConfig() override;
 
 	// ITitleFile
 	virtual TOnlineAsyncOpHandle<FTitleFileEnumerateFiles> EnumerateFiles(FTitleFileEnumerateFiles::Params&& Params) override;
diff --git a/Engine/Plugins/Online/Experimental/OnlineServicesNull/Source/Public/Online/UserFileNull.h b/Engine/Plugins/Online/Experimental/OnlineServicesNull/Source/Public/Online/UserFileNull.h
--- a/Engine/Plugins/Online/Experimental/OnlineServicesNull/Source/Public/Online/UserFileNull.h
+++ b/Engine/Plugins/Online/Experimental/OnlineServicesNull/Source/Public/Online/UserFileNull.h
@@ -17,7 +17,7 @@ public:
 	FUserFileNull(FOnlineServicesNull& InOwningSubsystem);
 
 	// IOnlineComponent
-	virtual void LoadConfig() override;
+	virtual void UpdateConfig() override;
 
 	// IUserFile
 	virtual TOnlineAsyncOpHandle<FUserFileEnumerateFiles> EnumerateFiles(FUserFileEnumerateFiles::Params&& Params) override;
diff --git a/Engine/Plugins/Online/Experimental/OnlineServicesOSSAdapter/Source/Private/Online/UserFileOSSAdapter.cpp b/Engine/Plugins/Online/Experimental/OnlineServicesOSSAdapter/Source/Private/Online/UserFileOSSAdapter.cpp
--- a/Engine/Plugins/Online/Experimental/OnlineServicesOSSAdapter/Source/Private/Online/UserFileOSSAdapter.cpp
+++ b/Engine/Plugins/Online/Experimental/OnlineServicesOSSAdapter/Source/Private/Online/UserFileOSSAdapter.cpp
@@ -23,9 +23,9 @@ void FUserFileOSSAdapter::PostInitialize()
 	check(UserCloudInterface);
 }
 
-void FUserFileOSSAdapter::LoadConfig()
+void FUserFileOSSAdapter::UpdateConfig()
 {
-	Super::LoadConfig();
+	Super::UpdateConfig();
 	Super::LoadConfig(Config);
 }
 
diff --git a/Engine/Plugins/Online/Experimental/OnlineServicesOSSAdapter/Source/Public/Online/UserFileOSSAdapter.h b/Engine/Plugins/Online/Experimental/OnlineServicesOSSAdapter/Source/Public/Online/UserFileOSSAdapter.h
--- a/Engine/Plugins/Online/Experimental/OnlineServicesOSSAdapter/Source/Public/Online/UserFileOSSAdapter.h
+++ b/Engine/Plugins/Online/Experimental/OnlineServicesOSSAdapter/Source/Public/Online/UserFileOSSAdapter.h
@@ -31,7 +31,7 @@ public:
 
 	// IOnlineComponent
 	virtual void PostInitialize() override;
-	virtual void LoadConfig() override;
+	virtual void UpdateConfig() override;
 
 	// IUserFile
 	virtual TOnlineAsyncOpHandle<FUserFileEnumerateFiles> EnumerateFiles(FUserFileEnumerateFiles::Params&& Params) override;
