Set default value for IOnlineServicesFactory::DefaultServiceOverride

diff --git a/Engine/Plugins/Online/Experimental/OnlineServices/Source/OnlineServicesInterface/Public/Online/OnlineServicesRegistry.h b/Engine/Plugins/Online/Experimental/OnlineServices/Source/OnlineServicesInterface/Public/Online/OnlineServicesRegistry.h
--- a/Engine/Plugins/Online/Experimental/OnlineServices/Source/OnlineServicesInterface/Public/Online/OnlineServicesRegistry.h
+++ b/Engine/Plugins/Online/Experimental/OnlineServices/Source/OnlineServicesInterface/Public/Online/OnlineServicesRegistry.h
@@ -134,7 +134,7 @@ private:
 
 	TMap<EOnlineServices, FFactoryAndPriority> ServicesFactories;
 	TMap<EOnlineServices, TMap<FName, TSharedRef<IOnlineServices>>> NamedServiceInstances;
-	EOnlineServices DefaultServiceOverride;
+	EOnlineServices DefaultServiceOverride = EOnlineServices::Default;
 
 	FOnlineServicesRegistry() {}
 	~FOnlineServicesRegistry();
