- Add Config struct support

diff --git a/Engine/Plugins/Online/Experimental/OnlineServices/Source/OnlineServicesCommon/Private/Online/OnlineConfig.cpp b/Engine/Plugins/Online/Experimental/OnlineServices/Source/OnlineServicesCommon/Private/Online/OnlineConfig.cpp
new file mode 100644
--- /dev/null
+++ b/Engine/Plugins/Online/Experimental/OnlineServices/Source/OnlineServicesCommon/Private/Online/OnlineConfig.cpp
@@ -0,0 +1,339 @@
+// Copyright Epic Games, Inc. All Rights Reserved.
+
+#include "Online/OnlineConfig.h"
+#include "Logging/LogMacros.h"
+#include "Misc/AsciiSet.h"
+#include "Misc/Char.h"
+
+DEFINE_LOG_CATEGORY_STATIC(LogOnlineServicesConfig, Log, VeryVerbose);
+
+namespace
+{
+using namespace UE::Online;
+
+constexpr FAsciiSet Equals(TEXT("="));
+constexpr FAsciiSet StringChars(TEXT("abcdefghijklmnopqrstuvwxyzABCDEFGHIJKLMNOPQRSTUVWXYZ0123456789_+"));
+	
+void SkipWhitespace(const TCHAR*& Str)
+{
+	while (FChar::IsWhitespace(*Str))
+	{
+		Str++;
+	}
+}
+
+const TCHAR* ParseField(const TCHAR* InStr, TMap<FString, FString>& OutMap)
+{
+	const TCHAR* Start = InStr;
+	const TCHAR* Str = FAsciiSet::FindFirstOrEnd(Start, Equals);
+	if (*Str)
+	{
+		int32 Len = UE_PTRDIFF_TO_INT32(Str - Start);
+		FString FieldName(Len, Start);
+
+		Str++;
+		Start = Str;
+		// Quoted string
+		if (*Str == TCHAR('\"'))
+		{
+			Start++;
+			do { Str++; } while (*Str && *Str != TCHAR('\"'));
+			if (*Str != TCHAR('\"'))
+			{
+				UE_LOG(LogOnlineServicesConfig, Warning, TEXT("Unmatched quoted string: %s"), InStr);
+				return nullptr;
+			}
+
+			Len = UE_PTRDIFF_TO_INT32(Str - Start);
+			FString FieldValue(Len, Start);
+
+			OutMap.Emplace(MoveTemp(FieldName), MoveTemp(FieldValue));
+			Str++;
+		}
+		// Struct
+		else if (*Str == TCHAR('('))
+		{
+			int SubCount = 0;
+			Str++;
+			while (*Str && (SubCount > 0 || *Str != TCHAR(')')))
+			{
+				SkipWhitespace(Str);
+				if (*Str == TCHAR('\"'))
+				{
+					do { Str++; } while (*Str && *Str != TCHAR('\"'));
+					if (*Str != TCHAR('\"'))
+					{
+						UE_LOG(LogOnlineServicesConfig, Warning, TEXT("Unmatched quoted string: %s"), InStr);
+						return nullptr;
+					}
+				}
+				else if (*Str == TCHAR('('))
+				{
+					SubCount++;
+				}
+				else if(*Str == TCHAR(')'))
+				{
+					SubCount--;
+					if (SubCount < 0)
+					{
+						UE_LOG(LogOnlineServicesConfig, Warning, TEXT("Too many closing parenthesis in: %s"), InStr);
+						return nullptr;
+					}
+				}
+				Str++;
+			}
+			Str++;
+
+			Len = UE_PTRDIFF_TO_INT32(Str - Start);
+			FString FieldValue(Len, Start);
+
+			OutMap.Emplace(MoveTemp(FieldName), MoveTemp(FieldValue));
+		}
+		// Unquoted string
+		else
+		{
+			Str = FAsciiSet::Skip(Str, StringChars);
+				
+			Len = UE_PTRDIFF_TO_INT32(Str - Start);
+			FString FieldValue(Len, Start);
+
+			OutMap.Emplace(MoveTemp(FieldName), MoveTemp(FieldValue));
+		}
+
+		// Skip comma.
+		if (*Str == TCHAR(','))
+		{
+			Str++;
+		}
+		SkipWhitespace(Str);
+	}
+
+	return Str;
+}
+}
+
+namespace UE::Online {
+
+IOnlineConfigStructPtr FOnlineConfigStructGConfig::CreateStruct(const FString& InConfigValue)
+{
+	IOnlineConfigStructPtr Result;
+	ParseStruct(GetData(InConfigValue), Result);
+	return Result;
+}
+
+TArray<IOnlineConfigStructPtr> FOnlineConfigStructGConfig::CreateStructArray(const FString& InConfigValue)
+{
+	TArray<IOnlineConfigStructPtr> Result;
+
+	if (!InConfigValue.IsEmpty())
+	{
+		const TCHAR* Str = GetData(InConfigValue);
+		// Start of array
+		if(*Str == TCHAR('('))
+		{
+			Str++;
+			// Start of array struct element
+			while (Str && *Str == TCHAR('('))
+			{	
+				IOnlineConfigStructPtr StructPtr;
+				Str = ParseStruct(Str, StructPtr);
+				if (StructPtr)
+				{
+					Result.Emplace(MoveTemp(StructPtr));
+				}
+			}
+			if (Str && *Str != TCHAR(')'))
+			{
+				UE_LOG(LogOnlineServicesConfig, Warning, TEXT("Unmatched opening parenthesis: %s"), Str);
+			}
+		}
+	}
+
+	return Result;
+}
+
+const TCHAR* FOnlineConfigStructGConfig::ParseStruct(const TCHAR* InStr, IOnlineConfigStructPtr& OutPtr)
+{
+	const TCHAR* Str = InStr;
+	if (*Str && *Str == TCHAR('('))
+	{
+		Str++;
+		FOnlineConfigStructGConfig Struct { FPrivateToken{} };
+		while (*Str && *Str != TCHAR(')'))
+		{
+			Str = ParseField(Str, Struct.StructMembers);
+			if (!Str)
+			{
+				return nullptr;
+			}
+		}
+		if (*Str != TCHAR(')'))
+		{
+			UE_LOG(LogOnlineServicesConfig, Warning, TEXT("Unmatched opening parenthesis: %s"), Str);
+			return nullptr;
+		}
+		Str++;
+		// Skip comma.
+		if (*Str == TCHAR(','))
+		{
+			Str++;
+		}
+		SkipWhitespace(Str);
+
+		OutPtr = MakeShared<FOnlineConfigStructGConfig>(MoveTemp(Struct));
+	}
+	else
+	{
+		UE_LOG(LogOnlineServicesConfig, Warning, TEXT("Failed to parse struct: %s"), Str);
+		return nullptr;
+	}
+
+	return Str;
+}
+
+TArray<FString> FOnlineConfigStructGConfig::CreateValueArray(const FString& InConfigValue)
+{
+	TArray<FString> Result;
+
+	if (!InConfigValue.IsEmpty())
+	{
+		const TCHAR* Str = GetData(InConfigValue);
+		// Start of array
+		if (*Str == TCHAR('('))
+		{
+			Str++;
+			while (*Str != TCHAR(')'))
+			{
+				const TCHAR* Start = Str;
+
+				// Quoted string
+				if (*Str == TCHAR('\"'))
+				{
+					Start++;
+					do { Str++; } while (*Str && *Str != TCHAR('\"'));
+					if (*Str != TCHAR('\"'))
+					{
+						UE_LOG(LogOnlineServicesConfig, Warning, TEXT("Unmatched quoted string: %s"), Str);
+						return Result;
+					}
+
+					const int32 Len = UE_PTRDIFF_TO_INT32(Str - Start);
+					FString FieldValue(Len, Start);
+					Result.Emplace(MoveTemp(FieldValue));
+					Str++;
+				}
+				// Unquoted string
+				else
+				{
+					Str = FAsciiSet::Skip(Str, StringChars);
+
+					const int32 Len = UE_PTRDIFF_TO_INT32(Str - Start);
+					FString FieldValue(Len, Start);
+					Result.Emplace(MoveTemp(FieldValue));
+				}
+
+				// Skip comma.
+				if (*Str == TCHAR(','))
+				{
+					Str++;
+				}
+				SkipWhitespace(Str);
+			}
+			if (*Str != TCHAR(')'))
+			{
+				UE_LOG(LogOnlineServicesConfig, Warning, TEXT("Unmatched opening parenthesis: %s"), Str);
+			}
+		}
+	}
+
+	return Result;
+}
+
+bool FOnlineConfigStructGConfig::GetValue(const TCHAR* Key, FString& Value)
+{
+	if (const FString* Found = StructMembers.Find(Key))
+	{
+		Value = *Found;
+		return true;
+	}
+	return false;
+}
+
+int32 FOnlineConfigStructGConfig::GetValue(const TCHAR* Key, TArray<FString>& Value)
+{
+	if (const FString* Found = StructMembers.Find(Key))
+	{
+		Value = CreateValueArray(*Found);
+		return Value.Num();
+	}
+	return 0;
+}
+
+bool FOnlineConfigStructGConfig::GetValue(const TCHAR* Key, IOnlineConfigStructPtr& Value)
+{
+	if (const FString* Found = StructMembers.Find(Key))
+	{
+		if (IOnlineConfigStructPtr ConfigStructPtr = CreateStruct(*Found))
+		{
+			Value = ConfigStructPtr;
+			return true;
+		}
+	}
+	return false;
+}
+
+int32 FOnlineConfigStructGConfig::GetValue(const TCHAR* Key, TArray<IOnlineConfigStructPtr>& Value)
+{
+	if (const FString* Found = StructMembers.Find(Key))
+	{
+		Value = CreateStructArray(*Found);
+		return Value.Num();
+	}
+	return 0;
+}
+
+bool FOnlineConfigProviderGConfig::GetValue(const TCHAR* Section, const TCHAR* Key, FString& Value)
+{
+	return GConfig->GetValue(Section, Key, Value, ConfigFile);
+}
+
+int32 FOnlineConfigProviderGConfig::GetValue(const TCHAR* Section, const TCHAR* Key, TArray<FString>& Value)
+{
+	return GConfig->GetValue(Section, Key, Value, ConfigFile);
+}
+
+bool FOnlineConfigProviderGConfig::GetValue(const TCHAR* Section, const TCHAR* Key, IOnlineConfigStructPtr& Value)
+{
+	FString ConfigValue;
+	if (GConfig->GetValue(Section, Key, ConfigValue, ConfigFile))
+	{
+		if (IOnlineConfigStructPtr StructPtr = FOnlineConfigStructGConfig::CreateStruct(ConfigValue))
+		{
+			Value = MoveTemp(StructPtr);
+			return true;
+		}
+	}
+
+	Value = nullptr;
+	return false;
+}
+
+int32 FOnlineConfigProviderGConfig::GetValue(const TCHAR* Section, const TCHAR* Key, TArray<IOnlineConfigStructPtr>& Value)
+{
+	Value.Empty();
+	TArray<FString> ConfigValues;
+	if (GConfig->GetValue(Section, Key, ConfigValues, ConfigFile))
+	{
+		Value.Reserve(ConfigValues.Num());
+		for (const FString& ConfigValue : ConfigValues)
+		{
+			if (IOnlineConfigStructPtr StructPtr = FOnlineConfigStructGConfig::CreateStruct(ConfigValue))
+			{
+				Value.Emplace(MoveTemp(StructPtr));
+			}
+		}
+	}
+	return Value.Num();
+}
+
+/* UE::Online */ }
diff --git a/Engine/Plugins/Online/Experimental/OnlineServices/Source/OnlineServicesCommon/Private/Online/OnlineSyntaxTests.cpp b/Engine/Plugins/Online/Experimental/OnlineServices/Source/OnlineServicesCommon/Private/Online/OnlineSyntaxTests.cpp
--- a/Engine/Plugins/Online/Experimental/OnlineServices/Source/OnlineServicesCommon/Private/Online/OnlineSyntaxTests.cpp
+++ b/Engine/Plugins/Online/Experimental/OnlineServices/Source/OnlineServicesCommon/Private/Online/OnlineSyntaxTests.cpp
@@ -403,21 +403,74 @@ inline void TestAsyncOpCacheSyntax(FOnlineServicesCommon& Services, Test::FTestI
 
 #if TEST_CONFIG_SYNTAX
 
+enum class EConfigTestEnum
+{
+	A,
+	B
+};
+void LexFromString(EConfigTestEnum& OutEnum, const TCHAR* InStr)
+{
+	if (FCString::Strcmp(InStr, TEXT("A")) == 0)
+	{
+		OutEnum = EConfigTestEnum::A;
+	}
+	else if (FCString::Strcmp(InStr, TEXT("B")) == 0)
+	{
+		OutEnum = EConfigTestEnum::B;
+	}
+}
+
+struct FTestConfigInnerInner
+{
+	float Float;
+};
+
+struct FTestConfigInner
+{
+	float Float;
+	FTestConfigInnerInner Object;
+	TArray<FTestConfigInnerInner> ObjectArray;
+	EConfigTestEnum Enum;
+	TArray<EConfigTestEnum> EnumArray;
+	TArray<FString> StringArray;
+};
+
 struct FTestConfig
 {
 	float Float;
 	int Int;
 	FString String;
 	TArray<int> IntArray;
+	FTestConfigInner Object;
+	TArray<FTestConfigInner> ObjectArray;
+	EConfigTestEnum Enum;
+	TArray<EConfigTestEnum> EnumArray;
 };
 
 namespace UE::Online::Meta {
 
+BEGIN_ONLINE_STRUCT_META(FTestConfigInnerInner)
+	ONLINE_STRUCT_FIELD(FTestConfigInnerInner, Float)
+END_ONLINE_STRUCT_META()
+
+BEGIN_ONLINE_STRUCT_META(FTestConfigInner)
+	ONLINE_STRUCT_FIELD(FTestConfigInner, Float),
+	ONLINE_STRUCT_FIELD(FTestConfigInner, Object),
+	ONLINE_STRUCT_FIELD(FTestConfigInner, ObjectArray),
+	ONLINE_STRUCT_FIELD(FTestConfigInner, Enum),
+	ONLINE_STRUCT_FIELD(FTestConfigInner, EnumArray),
+	ONLINE_STRUCT_FIELD(FTestConfigInner, StringArray)
+END_ONLINE_STRUCT_META()
+
 BEGIN_ONLINE_STRUCT_META(FTestConfig)
 	ONLINE_STRUCT_FIELD(FTestConfig, Float),
 	ONLINE_STRUCT_FIELD(FTestConfig, Int),
 	ONLINE_STRUCT_FIELD(FTestConfig, String),
-	ONLINE_STRUCT_FIELD(FTestConfig, IntArray)
+	ONLINE_STRUCT_FIELD(FTestConfig, IntArray),
+	ONLINE_STRUCT_FIELD(FTestConfig, Object),
+	ONLINE_STRUCT_FIELD(FTestConfig, ObjectArray),
+	ONLINE_STRUCT_FIELD(FTestConfig, Enum),
+	ONLINE_STRUCT_FIELD(FTestConfig, EnumArray)
 END_ONLINE_STRUCT_META()
 
 /* UE::Online::Meta*/ }
@@ -425,7 +478,10 @@ END_ONLINE_STRUCT_META()
 void TestLoadConfigSyntax(UE::Online::IOnlineConfigProvider& ConfigProvider)
 {
 	FTestConfig TestConfig;
+	// Test loading from a section (each field maps to a key in that section)
 	UE::Online::LoadConfig(ConfigProvider, TEXT("Test"), TestConfig);
+	// Test loading from a specific key in a section (the config value is in object syntax, each field maps to a key in that object)
+	UE::Online::LoadConfig(ConfigProvider, TEXT("Test"), TEXT("Test"), TestConfig);
 }
 
 #endif // TEST_CONFIG_SYNTAX
diff --git a/Engine/Plugins/Online/Experimental/OnlineServices/Source/OnlineServicesCommon/Public/Online/OnlineConfig.h b/Engine/Plugins/Online/Experimental/OnlineServices/Source/OnlineServicesCommon/Public/Online/OnlineConfig.h
--- a/Engine/Plugins/Online/Experimental/OnlineServices/Source/OnlineServicesCommon/Public/Online/OnlineConfig.h
+++ b/Engine/Plugins/Online/Experimental/OnlineServices/Source/OnlineServicesCommon/Public/Online/OnlineConfig.h
@@ -4,10 +4,154 @@
 
 #include "Misc/ConfigCacheIni.h"
 #include "Online/OnlineMeta.h"
+#include "Online/SchemaTypes.h"
 #include "Algo/Transform.h"
 
 namespace UE::Online {
 
+namespace Private {
+
+template <typename T> constexpr bool MissingCase = false;
+
+template <typename T>
+T ParseValueFromString(const FString& StringValue)
+{
+	if constexpr (std::is_same_v<T, FText>)
+	{
+		FText TextValue;
+		FTextStringHelper::ReadFromBuffer(*StringValue, TextValue);
+		return TextValue;
+	}
+	else if constexpr (std::is_same_v<T, FName>)
+	{
+		return FName(*StringValue);
+	}
+	else if constexpr (std::is_same_v<T, bool>)
+	{
+		return FCString::ToBool(*StringValue);
+	}
+	else if constexpr (std::is_same_v<T, int32>)
+	{
+		return FCString::Strtoi(*StringValue, nullptr, 10);
+	}
+	else if constexpr (std::is_same_v<T, uint32>)
+	{
+		return (uint32)FCString::Strtoui64(*StringValue, nullptr, 10);
+	}
+	else if constexpr (std::is_same_v<T, int64>)
+	{
+		return FCString::Strtoi64(*StringValue, nullptr, 10);
+	}
+	else if constexpr (std::is_same_v<T, uint64>)
+	{
+		return FCString::Strtoui64(*StringValue, nullptr, 10);
+	}
+	else if constexpr (std::is_same_v<T, float>)
+	{
+		return FCString::Atof(*StringValue);
+	}
+	else if constexpr (std::is_same_v<T, double>)
+	{
+		return FCString::Atod(*StringValue);
+	}
+	else if constexpr (std::is_enum_v<T>
+		|| std::is_same_v<T, FSchemaVariant>)
+	{
+		T OutValue;
+		using ::LexFromString;
+		LexFromString(OutValue, *StringValue);
+		return OutValue;
+	}
+	else
+	{
+		static_assert(MissingCase<T>, "ParseValueFromString missing implementation for type T");
+	}
+}
+
+/* Private */ }
+
+class IOnlineConfigStruct;
+using IOnlineConfigStructPtr = TSharedPtr<IOnlineConfigStruct>;
+
+class ONLINESERVICESCOMMON_API IOnlineConfigStruct
+{
+public:
+	virtual ~IOnlineConfigStruct() = default;
+
+	virtual bool GetValue(const TCHAR* Key, FString& Value) = 0;
+	virtual int32 GetValue(const TCHAR* Key, TArray<FString>& Value) = 0;
+	virtual bool GetValue(const TCHAR* Key, IOnlineConfigStructPtr& Value) = 0;
+	virtual int32 GetValue(const TCHAR* Key, TArray<IOnlineConfigStructPtr>& Value) = 0;
+
+	template <typename T>
+	bool GetValue(const TCHAR* Key, T& Value)
+	{
+		if constexpr (TModels<Meta::COnlineMetadataAvailable, T>::Value)
+		{
+			IOnlineConfigStructPtr ConfigStructPtr;
+			if (GetValue(Key, ConfigStructPtr))
+			{
+				T Result;
+				Meta::VisitFields(Result,
+					[&ConfigStructPtr](const TCHAR* FieldName, auto& Field)
+				{
+					ConfigStructPtr->GetValue(FieldName, Field);
+				});
+				Value = MoveTemp(Result);
+				return true;
+			}
+		}
+		else
+		{
+			FString StringValue;
+			if (GetValue(Key, StringValue))
+			{
+				Value = Private::ParseValueFromString<T>(StringValue);
+				return true;
+			}
+		}
+		return false;
+	}
+
+	template <typename T>
+	int32 GetValue(const TCHAR* Key, TArray<T>& Value)
+	{
+		Value.Empty();
+
+		if constexpr (TModels<Meta::COnlineMetadataAvailable, T>::Value)
+		{
+			TArray<IOnlineConfigStructPtr> ConfigStructPtrArray;
+			if (GetValue(Key, ConfigStructPtrArray) > 0)
+			{
+				Value.Reserve(ConfigStructPtrArray.Num());
+				for (const IOnlineConfigStructPtr& ConfigStructPtr : ConfigStructPtrArray)
+				{
+					T& Struct = Value.Emplace_GetRef();
+					Meta::VisitFields(Struct,
+						[&ConfigStructPtr](const TCHAR* FieldName, auto& Field)
+					{
+						ConfigStructPtr->GetValue(FieldName, Field);
+					});
+				}
+			}
+		}
+		else
+		{
+			TArray<FString> StringArray;
+			if (GetValue(Key, StringArray) > 0)
+			{
+				Value.Reserve(StringArray.Num());
+				for (const FString& StringValue : StringArray)
+				{
+					Value.Emplace(Private::ParseValueFromString<T>(StringValue));
+				}
+			}
+		}
+
+		return Value.Num();
+	}
+};
+
 /**
  * Interface for retrieving config values used by OnlineServices implementations
  * 
@@ -55,372 +199,105 @@ public:
 	 */
 	virtual int32 GetValue(const TCHAR* Section, const TCHAR* Key, TArray<FString>& Value) = 0;
 
-	/**
-	 * Get an FText value
-	 *
-	 * @param Section Section to read the value from
-	 * @param Key Key in the section for the value
-	 * @param Value FText value read from the config. Will be unchanged if not present
-	 *
-	 * @return true if a value was read
-	 */
-	virtual bool GetValue(const TCHAR* Section, const TCHAR* Key, FText& Value)
-	{
-		FString StringValue;
-		if (GetValue(Section, Key, StringValue))
-		{
-			return FTextStringHelper::ReadFromBuffer(*StringValue, Value, Section) != nullptr;
-		}
-		return false;
-	}
+	virtual bool GetValue(const TCHAR* Section, const TCHAR* Key, IOnlineConfigStructPtr& Value) = 0;
 
-	/**
-	 * Get an FName value
-	 *
-	 * @param Section Section to read the value from
-	 * @param Key Key in the section for the value
-	 * @param Value FName value read from the config. Will be unchanged if not present
-	 *
-	 * @return true if a value was read
-	 */
-	virtual bool GetValue(const TCHAR* Section, const TCHAR* Key, FName& Value)
-	{
-		FString StringValue;
-		if (GetValue(Section, Key, StringValue))
-		{
-			Value = FName(*StringValue);
-			return true;
-		}
-		return false;
-	}
-
-	/**
-	 * Get a bool value
-	 *
-	 * @param Section Section to read the value from
-	 * @param Key Key in the section for the value
-	 * @param Value bool value read from the config. Will be unchanged if not present
-	 *
-	 * @return true if a value was read
-	 */
-	virtual bool GetValue(const TCHAR* Section, const TCHAR* Key, bool& Value)
-	{
-		FString StringValue;
-		if (GetValue(Section, Key, StringValue))
-		{
-			Value = FCString::ToBool(*StringValue);
-			return true;
-		}
-		return false;
-	}
-
-	/**
-	 * Get an int32 value
-	 *
-	 * @param Section Section to read the value from
-	 * @param Key Key in the section for the value
-	 * @param Value int32 value read from the config. Will be unchanged if not present
-	 *
-	 * @return true if a value was read
-	 */
-	virtual bool GetValue(const TCHAR* Section, const TCHAR* Key, int32& Value)
-	{
-		FString StringValue;
-		if (GetValue(Section, Key, StringValue))
-		{
-			Value = FCString::Strtoi(*StringValue, nullptr, 10);
-			return true;
-		}
-		return false;
-	}
-
-	/**
-	 * Get an int64 value
-	 *
-	 * @param Section Section to read the value from
-	 * @param Key Key in the section for the value
-	 * @param Value int64 value read from the config. Will be unchanged if not present
-	 *
-	 * @return true if a value was read
-	 */
-	virtual bool GetValue(const TCHAR* Section, const TCHAR* Key, int64& Value)
-	{
-		FString StringValue;
-		if (GetValue(Section, Key, StringValue))
-		{
-			Value = FCString::Strtoi64(*StringValue, nullptr, 10);
-			return true;
-		}
-		return false;
-	}
+	virtual int32 GetValue(const TCHAR* Section, const TCHAR* Key, TArray<IOnlineConfigStructPtr>& Value) = 0;
 
-	/**
-	 * Get a uint64 value
-	 *
-	 * @param Section Section to read the value from
-	 * @param Key Key in the section for the value
-	 * @param Value uint64 value read from the config. Will be unchanged if not present
-	 *
-	 * @return true if a value was read
-	 */
-	virtual bool GetValue(const TCHAR* Section, const TCHAR* Key, uint64& Value)
-	{
-		FString StringValue;
-		if (GetValue(Section, Key, StringValue))
-		{
-			Value = FCString::Strtoui64(*StringValue, nullptr, 10);
-			return true;
-		}
-		return false;
-	}
-
-	/**
-	 * Get a float value
-	 *
-	 * @param Section Section to read the value from
-	 * @param Key Key in the section for the value
-	 * @param Value float value read from the config. Will be unchanged if not present
-	 *
-	 * @return true if a value was read
-	 */
-	virtual bool GetValue(const TCHAR* Section, const TCHAR* Key, float& Value)
+	template <typename T>
+	bool GetValue(const TCHAR* Section, const TCHAR* Key, T& Value)
 	{
-		FString StringValue;
-		if (GetValue(Section, Key, StringValue))
+		if constexpr (TModels<Meta::COnlineMetadataAvailable, T>::Value)
 		{
-			Value = FCString::Atof(*StringValue);
-			return true;
+			IOnlineConfigStructPtr ConfigStructPtr;
+			if (GetValue(Section, Key, ConfigStructPtr))
+			{
+				T Result;
+				Meta::VisitFields(Result,
+					[&ConfigStructPtr](const TCHAR* FieldName, auto& Field)
+				{
+					ConfigStructPtr->GetValue(FieldName, Field);
+				});
+				Value = MoveTemp(Result);
+				return true;
+			}
 		}
-		return false;
-	}
-
-	/**
-	 * Get a double value
-	 *
-	 * @param Section Section to read the value from
-	 * @param Key Key in the section for the value
-	 * @param Value double value read from the config. Will be unchanged if not present
-	 *
-	 * @return true if a value was read
-	 */
-	virtual bool GetValue(const TCHAR* Section, const TCHAR* Key, double& Value)
-	{
-		FString StringValue;
-		if (GetValue(Section, Key, StringValue))
+		else
 		{
-			Value = FCString::Atod(*StringValue);
-			return true;
+			FString StringValue;
+			if (GetValue(Section, Key, StringValue))
+			{
+				Value = Private::ParseValueFromString<T>(StringValue);
+				return true;
+			}
 		}
 		return false;
 	}
 
-	/**
-	 * Get a value consisting of a TArray of FText
-	 *
-	 * @param Section Section to read the value from
-	 * @param Key Key in the section for the value
-	 * @param Value Array of FText read from the config. Will be empty if not present
-	 *
-	 * @return number of values in the array
-	 */
-	virtual int32 GetValue(const TCHAR* Section, const TCHAR* Key, TArray<FText>& Value)
+	template <typename T>
+	int32 GetValue(const TCHAR* Section, const TCHAR* Key, TArray<T>& Value)
 	{
-		Value.Empty();
-		TArray<FString> StringArray;
-		if (GetValue(Section, Key, StringArray) > 0)
+		int32 NumInitialElements = Value.Num();
+		if constexpr (TModels<Meta::COnlineMetadataAvailable, T>::Value)
 		{
-			Value.Reserve(StringArray.Num());
-			Algo::Transform(StringArray, Value,
-				[Section](const FString& StringValue) -> FText
+			TArray<IOnlineConfigStructPtr> ConfigStructPtrArray;
+			if (GetValue(Section, Key, ConfigStructPtrArray) > 0)
+			{
+				Value.Reserve(ConfigStructPtrArray.Num() + NumInitialElements);
+				for (const IOnlineConfigStructPtr& ConfigStructPtr : ConfigStructPtrArray)
 				{
-					FText TextValue;
-					FTextStringHelper::ReadFromBuffer(*StringValue, TextValue, Section);
-					return TextValue;
-				});
+					T& Struct = Value.Emplace_GetRef();
+					Meta::VisitFields(Struct,
+						[&ConfigStructPtr](const TCHAR* FieldName, auto& Field)
+					{
+						ConfigStructPtr->GetValue(FieldName, Field);
+					});
+				}
+			}
 		}
-		return Value.Num();
-	}
-
-	/**
-	 * Get a value consisting of a TArray of FName
-	 *
-	 * @param Section Section to read the value from
-	 * @param Key Key in the section for the value
-	 * @param Value Array of FName values read from the config. Will be empty if not present
-	 *
-	 * @return number of values in the array
-	 */
-	virtual int32 GetValue(const TCHAR* Section, const TCHAR* Key, TArray<FName>& Value)
-	{
-		Value.Empty();
-		TArray<FString> StringArray;
-		if (GetValue(Section, Key, StringArray) > 0)
+		else
 		{
-			Value.Reserve(StringArray.Num());
-			Algo::Transform(StringArray, Value,
-				[](const FString& StringValue) -> const TCHAR*
+			TArray<FString> StringArray;
+			if (GetValue(Section, Key, StringArray) > 0)
+			{
+				Value.Reserve(StringArray.Num() + NumInitialElements);
+				for (const FString& StringValue : StringArray)
 				{
-					return *StringValue;
-				});
+					Value.Emplace(Private::ParseValueFromString<T>(StringValue));
+				}
+			}
 		}
-		return Value.Num();
-	}
 
-	/**
-	 * Get a value consisting of a TArray of bool
-	 *
-	 * @param Section Section to read the value from
-	 * @param Key Key in the section for the value
-	 * @param Value Array of bool read from the config. Will be empty if not present
-	 *
-	 * @return number of values in the array
-	 */
-	virtual int32 GetValue(const TCHAR* Section, const TCHAR* Key, TArray<bool>& Value)
-	{
-		Value.Empty();
-		TArray<FString> StringArray;
-		if (GetValue(Section, Key, StringArray) > 0)
-		{
-			Value.Reserve(StringArray.Num());
-			Algo::Transform(StringArray, Value,
-				[](const FString& StringValue) -> bool
-				{
-					return FCString::ToBool(*StringValue);
-				});
-		}
-		return Value.Num();
+		return Value.Num() - NumInitialElements;
 	}
+};
 
-	/**
-	 * Get a value consisting of a TArray of int32
-	 *
-	 * @param Section Section to read the value from
-	 * @param Key Key in the section for the value
-	 * @param Value Array of int32 read from the config. Will be empty if not present
-	 *
-	 * @return number of values in the array
-	 */
-	virtual int32 GetValue(const TCHAR* Section, const TCHAR* Key, TArray<int32>& Value)
-	{
-		Value.Empty();
-		TArray<FString> StringArray;
-		if (GetValue(Section, Key, StringArray) > 0)
-		{
-			Value.Reserve(StringArray.Num());
-			Algo::Transform(StringArray, Value,
-				[](const FString& StringValue) -> int32
-				{
-					return FCString::Strtoi(*StringValue, nullptr, 10);
-				});
-		}
-		return Value.Num();
-	}
+class ONLINESERVICESCOMMON_API FOnlineConfigStructGConfig : public IOnlineConfigStruct
+{
+public:
+	virtual ~FOnlineConfigStructGConfig() = default;
 
-	/**
-	 * Get a value consisting of a TArray of int64
-	 *
-	 * @param Section Section to read the value from
-	 * @param Key Key in the section for the value
-	 * @param Value Array of int64 read from the config. Will be empty if not present
-	 *
-	 * @return number of values in the array
-	 */
-	virtual int32 GetValue(const TCHAR* Section, const TCHAR* Key, TArray<int64>& Value)
-	{
-		Value.Empty();
-		TArray<FString> StringArray;
-		if (GetValue(Section, Key, StringArray) > 0)
-		{
-			Value.Reserve(StringArray.Num());
-			Algo::Transform(StringArray, Value,
-				[](const FString& StringValue) -> int64
-				{
-					return FCString::Strtoi64(*StringValue, nullptr, 10);
-				});
-		}
-		return Value.Num();
-	}
+	virtual bool GetValue(const TCHAR* Key, FString& Value) override;
+	virtual int32 GetValue(const TCHAR* Key, TArray<FString>& Value) override;
+	virtual bool GetValue(const TCHAR * Key, IOnlineConfigStructPtr& Value) override;
+	virtual int32 GetValue(const TCHAR* Key, TArray<IOnlineConfigStructPtr>& Value) override;
 
-	/**
-	 * Get a value consisting of a TArray of uint64
-	 *
-	 * @param Section Section to read the value from
-	 * @param Key Key in the section for the value
-	 * @param Value Array of uint64 read from the config. Will be empty if not present
-	 *
-	 * @return number of values in the array
-	 */
-	virtual int32 GetValue(const TCHAR* Section, const TCHAR* Key, TArray<uint64>& Value)
-	{
-		Value.Empty();
-		TArray<FString> StringArray;
-		if (GetValue(Section, Key, StringArray) > 0)
-		{
-			Value.Reserve(StringArray.Num());
-			Algo::Transform(StringArray, Value,
-				[](const FString& StringValue) -> uint64
-				{
-					return FCString::Strtoui64(*StringValue, nullptr, 10);
-				});
-		}
-		return Value.Num();
-	}
+private:
+	struct FPrivateToken { explicit FPrivateToken() = default; };
 
-	/**
-	 * Get a value consisting of a TArray of float
-	 *
-	 * @param Section Section to read the value from
-	 * @param Key Key in the section for the value
-	 * @param Value Array of float read from the config. Will be empty if not present
-	 *
-	 * @return number of values in the array
-	 */
-	virtual int32 GetValue(const TCHAR* Section, const TCHAR* Key, TArray<float>& Value)
+public:
+	FOnlineConfigStructGConfig(FPrivateToken)
 	{
-		Value.Empty();
-		TArray<FString> StringArray;
-		if (GetValue(Section, Key, StringArray) > 0)
-		{
-			Value.Reserve(StringArray.Num());
-			Algo::Transform(StringArray, Value,
-				[](const FString& StringValue) -> float
-				{
-					return FCString::Atof(*StringValue);
-				});
-		}
-		return Value.Num();
 	}
 
-	/**
-	 * Get a value consisting of a TArray of double
-	 *
-	 * @param Section Section to read the value from
-	 * @param Key Key in the section for the value
-	 * @param Value Array of double read from the config. Will be empty if not present
-	 *
-	 * @return number of values in the array
-	 */
-	virtual int32 GetValue(const TCHAR* Section, const TCHAR* Key, TArray<double>& Value)
-	{
-		Value.Empty();
-		TArray<FString> StringArray;
-		if (GetValue(Section, Key, StringArray) > 0)
-		{
-			Value.Reserve(StringArray.Num());
-			Algo::Transform(StringArray, Value,
-				[](const FString& StringValue) -> double
-				{
-					return FCString::Atod(*StringValue);
-				});
-		}
-		return Value.Num();
-	}
+private:
+	friend class FOnlineConfigProviderGConfig;
 
-	template <typename T>
-	std::enable_if_t<TModels<Meta::COnlineMetadataAvailable, T>::Value, bool> GetValue(const TCHAR* Section, const TCHAR* Key, T& Value);
+	static IOnlineConfigStructPtr CreateStruct(const FString& InConfigValue);
+	static TArray<IOnlineConfigStructPtr> CreateStructArray(const FString& InConfigValue);
+	static const TCHAR* ParseStruct(const TCHAR* InStr, IOnlineConfigStructPtr& OutPtr);
+	static TArray<FString> CreateValueArray(const FString& InConfigValue);
 
-	template <typename T>
-	std::enable_if_t<TModels<Meta::COnlineMetadataAvailable, T>::Value, int32> GetValue(const TCHAR* Section, const TCHAR* Key, TArray<T>& Value);
+	TMap<FString, FString> StructMembers;
 };
 
 /**
@@ -433,16 +310,12 @@ public:
 		: ConfigFile(InConfigFile)
 	{
 	}
+	virtual ~FOnlineConfigProviderGConfig() = default;
 
-	virtual bool GetValue(const TCHAR* Section, const TCHAR* Key, FString& Value) override
-	{
-		return GConfig->GetValue(Section, Key, Value, ConfigFile);
-	}
-
-	virtual int32 GetValue(const TCHAR* Section, const TCHAR* Key, TArray<FString>& Value) override
-	{
-		return GConfig->GetValue(Section, Key, Value, ConfigFile);
-	}
+	virtual bool GetValue(const TCHAR* Section, const TCHAR* Key, FString& Value) override;
+	virtual int32 GetValue(const TCHAR* Section, const TCHAR* Key, TArray<FString>& Value) override;
+	virtual bool GetValue(const TCHAR* Section, const TCHAR* Key, IOnlineConfigStructPtr& Value) override;
+	virtual int32 GetValue(const TCHAR* Section, const TCHAR* Key, TArray<IOnlineConfigStructPtr>& Value) override;
 
 private:
 	FString ConfigFile;
@@ -450,141 +323,133 @@ private:
 
 namespace Private {
 
-// If needed, this can be specialized for additional types generically using LexFromString or explicit specialization
-template <typename T>
-auto LoadConfigValue(IOnlineConfigProvider& Provider, const TCHAR* Section, const TCHAR* Key, T& Value)
-	-> std::enable_if_t<std::is_same_v<decltype(Provider.GetValue(Section, Key, Value)), bool>, bool>
-{
-	return Provider.GetValue(Section, Key, Value);
-}
-
 template <typename T>
-auto LoadConfigValue(IOnlineConfigProvider& Provider, const TCHAR* Section, const TCHAR* Key, T& Value)
-	-> std::enable_if_t<std::is_same_v<decltype(Provider.GetValue(Section, Key, Value)), int32>, bool>
+bool LoadConfigStructMemberValue(IOnlineConfigStruct& ConfigStruct, const TCHAR* Key, T& Value)
 {
-	return Provider.GetValue(Section, Key, Value) > 0;
-}
-
-template <typename T>
-auto LoadConfigValue(IOnlineConfigProvider& Provider, const TCHAR* Section, const TCHAR* Key, T& Value)
-	-> std::enable_if_t<std::is_enum_v<T>, bool>
-{
-	FString StringValue;
-	if (Provider.GetValue(Section, Key, StringValue))
+	if constexpr (std::is_same_v<decltype(ConfigStruct.GetValue(Key, Value)), bool>)
 	{
-		using ::LexFromString;
-		LexFromString(Value, *StringValue);
-		return true;
+		return ConfigStruct.GetValue(Key, Value);
+	}
+	else if constexpr (std::is_same_v<decltype(ConfigStruct.GetValue(Key, Value)), int32>)
+	{
+		return ConfigStruct.GetValue(Key, Value) > 0;
 	}
-	return false;
 }
 
 template <typename T>
-auto LoadConfigValue(IOnlineConfigProvider& Provider, const TCHAR* Section, const TCHAR* Key, TArray<T>& Value)
-	-> std::enable_if_t<std::is_enum_v<T>, int32>
+bool LoadConfigValue(IOnlineConfigProvider& Provider, const TCHAR* Section, const TCHAR* Key, T& Value)
 {
-	Value.Empty();
-	TArray<FString> StringArray;
-	if (Provider.GetValue(Section, Key, StringArray))
+	if constexpr (std::is_same_v<decltype(Provider.GetValue(Section, Key, Value)), bool>)
 	{
-		Value.Reserve(StringArray.Num());
-		Algo::Transform(StringArray, Value,
-			[](const FString& StringValue) -> T
-			{
-				T EnumValue;
-				using ::LexFromString;
-				LexFromString(EnumValue, *StringValue);
-				return EnumValue;
-			});
+		return Provider.GetValue(Section, Key, Value);
+	}
+	else if constexpr (std::is_same_v<decltype(Provider.GetValue(Section, Key, Value)), int32>)
+	{
+		return Provider.GetValue(Section, Key, Value) > 0;
 	}
-	return Value.Num();
 }
 
 /* Private */ }
 
 /**
- * Populate a struct from a config provider. It will load from keys matching the struct member names in the section
+ * Populate a struct from a config provider. Loads each field in the struct by matching field names with keys in the config section.
  * Requires that the type metadata is specified for StructType
- * 
+ *
  * @param Provider The config provider
  * @param Section The section of the config to load the values from
- * @param The struct whose fields will be loaded from config
- * 
+ * @param Value The struct whose fields will be loaded from config
+ *
  * @return true if any values were loaded
  */
-template <typename StructType>
-bool LoadConfig(IOnlineConfigProvider& Provider, const FString& Section, StructType& Struct)
+template <typename T>
+auto LoadConfig(IOnlineConfigProvider& Provider, const FString& Section, T& Value)
+	-> std::enable_if_t<TModels<Meta::COnlineMetadataAvailable, T>::Value, bool>
 {
 	bool bLoadedValue = false;
-	Meta::VisitFields(Struct,
+	Meta::VisitFields(Value,
 		[&Provider, &Section, &bLoadedValue](const TCHAR* FieldName, auto& Field)
-		{
-			bLoadedValue |= Private::LoadConfigValue(Provider, *Section, FieldName, Field);
-		});
+	{
+		bLoadedValue |= Private::LoadConfigValue(Provider, *Section, FieldName, Field);
+	});
 	return bLoadedValue;
 }
 
-template <typename StructType>
-bool LoadConfig(IOnlineConfigProvider& Provider, const TArray<FString>& SectionHeiarchy, StructType& OutValue)
+/**
+ * Populate a struct from a config provider. Iterates through each section in SectionHierarchy, loading each field in the struct by matching field names with keys in the section.
+ * Requires that the type metadata is specified for StructType
+ *
+ * @param Provider The config provider
+ * @param SectionHeirarchy The config sections to load values from
+ * @param Value The struct whose fields will be loaded from config
+ *
+ * @return true if any values were loaded
+ */
+template <typename T>
+auto LoadConfig(IOnlineConfigProvider& Provider, const TArray<FString>& SectionHeirarchy, T& Value)
+	-> std::enable_if_t<TModels<Meta::COnlineMetadataAvailable, T>::Value, bool>
 {
 	bool bLoadedValue = false;
-	for (const FString& Section : SectionHeiarchy)
+	for (const FString& Section : SectionHeirarchy)
 	{
-		StructType Value;
-		if (LoadConfig(Provider, Section, Value))
-		{
-			bLoadedValue = true;
-			OutValue = MoveTemp(Value);
-		}
+		bLoadedValue |= LoadConfig(Provider, Section, Value);
 	}
 	return bLoadedValue;
 }
 
+/**
+ * Populate a struct from a config provider. Loads the struct from the _value_ of the given key in the given config section.
+ * Requires that the type metadata is specified for StructType, and that the config value is in config struct syntax.
+ * 
+ * @param Provider The config provider
+ * @param Section The section of the config where the struct resides
+ * @param Key The key in the config section where the struct resides
+ * @param Value The struct whose fields will be loaded from config
+ * 
+ * @return true if any values were loaded
+ */
 template <typename T>
-std::enable_if_t<TModels<Meta::COnlineMetadataAvailable, T>::Value, bool> IOnlineConfigProvider::GetValue(const TCHAR* Section, const TCHAR* Key, T& OutValue)
+auto LoadConfig(IOnlineConfigProvider& Provider, const FString& Section, const TCHAR* Key, T& OutValue)
+-> std::enable_if_t<TModels<Meta::COnlineMetadataAvailable, T>::Value, bool>
 {
-	const FString AggregateSection = FString::Printf(TEXT("%s.%s"), Section, Key);
 	bool bLoadedValue = false;
-	T Value;
-	Meta::VisitFields(Value,
-		[this, &AggregateSection, &bLoadedValue](const TCHAR* FieldName, auto& Field)
+	IOnlineConfigStructPtr ConfigStructPtr;
+	if (Provider.GetValue(*Section, Key, ConfigStructPtr))
+	{
+		T Value;
+		Meta::VisitFields(Value,
+			[&ConfigStructPtr, &bLoadedValue](const TCHAR* FieldName, auto& Field)
 		{
-			bLoadedValue |= Private::LoadConfigValue(*this, *AggregateSection, FieldName, Field);
+			bLoadedValue |= Private::LoadConfigStructMemberValue(*ConfigStructPtr, FieldName, Field);
 		});
-
-	if (bLoadedValue)
-	{
-		OutValue = MoveTemp(Value);
+		if (bLoadedValue)
+		{
+			OutValue = MoveTemp(Value);
+		}
 	}
-
 	return bLoadedValue;
 }
 
+/**
+ * Populate a struct from a config provider. Iterates through each section in SectionHierarchy, loading the struct from the _value_ of the given key in the given config section.
+ * Requires that the type metadata is specified for StructType, and that the config value is in config struct syntax.
+ *
+ * @param Provider The config provider
+ * @param SectionHeirarchy The config sections to load values from
+ * @param Key The key in the config section where the struct resides
+ * @param Value The struct whose fields will be loaded from config
+ *
+ * @return true if any values were loaded
+ */
 template <typename T>
-std::enable_if_t<TModels<Meta::COnlineMetadataAvailable, T>::Value, int32> IOnlineConfigProvider::GetValue(const TCHAR* Section, const TCHAR* Key, TArray<T>& Value)
+auto LoadConfig(IOnlineConfigProvider& Provider, const TArray<FString>& SectionHeirarchy, const TCHAR* Key, T& Value)
+-> std::enable_if_t<TModels<Meta::COnlineMetadataAvailable, T>::Value, bool>
 {
-	Value.Empty();
-	TArray<FString> StringArray;
-	if (GetValue(Section, Key, StringArray) > 0)
+	bool bLoadedValue = false;
+	for (const FString& Section : SectionHeirarchy)
 	{
-		Value.Reserve(StringArray.Num());
-		for (const FString& StringValue : StringArray)
-		{
-			FString PrecedingToken;
-			const TCHAR* OldValue = *StringValue;
-
-			if (FParse::Token(OldValue, PrecedingToken, true))
-			{
-				const FString AggregateSection = FString::Printf(TEXT("%s.%s.%s"), Section, Key, *StringValue);
-				Meta::VisitFields(Value.Emplace_GetRef(),
-					[this, &AggregateSection](const TCHAR* FieldName, auto& Field)
-					{
-						Private::LoadConfigValue(*this, *AggregateSection, FieldName, Field);
-					});
-			}
-		}
+		bLoadedValue |= LoadConfig(Provider, Section, Key, Value);
 	}
-	return Value.Num();
+	return bLoadedValue;
 }
 
 /* UE::Online */ }
diff --git a/Engine/Plugins/Online/Experimental/OnlineServices/Source/OnlineServicesInterface/Private/Online/SchemaTypes.cpp b/Engine/Plugins/Online/Experimental/OnlineServices/Source/OnlineServicesInterface/Private/Online/SchemaTypes.cpp
--- a/Engine/Plugins/Online/Experimental/OnlineServices/Source/OnlineServicesInterface/Private/Online/SchemaTypes.cpp
+++ b/Engine/Plugins/Online/Experimental/OnlineServices/Source/OnlineServicesInterface/Private/Online/SchemaTypes.cpp
@@ -394,7 +394,32 @@ const FString LexToString(const FLexToStringAdaptor<FSchemaVariant>& Adaptor)
 
 void LexFromString(FSchemaVariant& OutSchemaVariant, const TCHAR* InStr)
 {
-	// todo
+	if (const TCHAR* ValueStr = FCString::Strchr(InStr, ':'))
+	{
+		const int32 TypeLen = UE_PTRDIFF_TO_INT32(ValueStr - InStr);
+		const FString TypeStr(TypeLen, InStr);
+		ValueStr++;
+
+		LexFromString(OutSchemaVariant.VariantType, *TypeStr);
+		switch (OutSchemaVariant.VariantType)
+		{
+		case ESchemaAttributeType::None:
+			OutSchemaVariant.VariantData = FSchemaVariant::FVariantType();
+			break;
+		case ESchemaAttributeType::Bool:
+			OutSchemaVariant.VariantData.Emplace<bool>(FCString::ToBool(ValueStr));
+			break;
+		case ESchemaAttributeType::Int64:
+			OutSchemaVariant.VariantData.Emplace<int64>(FCString::Strtoi64(ValueStr, nullptr, 10));
+			break;
+		case ESchemaAttributeType::Double:
+			OutSchemaVariant.VariantData.Emplace<double>(FCString::Atod(ValueStr));
+			break;
+		case ESchemaAttributeType::String:
+			OutSchemaVariant.VariantData.Emplace<FString>(ValueStr);
+			break;
+		}
+	}
 }
 
 /* UE::Online */ }
