- Refactor Schema support to work new config struct loading, and not be too onerous to modify with huge one-liner schema definitions.
- Mostly splitting up some of the nested types into separate arrays.
- Fix tests, and add some new ones for AdditionalCategoryAttributes.

diff --git a/Engine/Plugins/Online/Experimental/OnlineServices/Source/OnlineServicesCommon/Private/Online/OnlineServicesCommon.cpp b/Engine/Plugins/Online/Experimental/OnlineServices/Source/OnlineServicesCommon/Private/Online/OnlineServicesCommon.cpp
--- a/Engine/Plugins/Online/Experimental/OnlineServices/Source/OnlineServicesCommon/Private/Online/OnlineServicesCommon.cpp
+++ b/Engine/Plugins/Online/Experimental/OnlineServices/Source/OnlineServicesCommon/Private/Online/OnlineServicesCommon.cpp
@@ -24,9 +24,9 @@ FOnlineServicesCommon::FOnlineServicesCommon(const FString& InConfigName, FName
 void FOnlineServicesCommon::Init()
 {
 	OpCache.SetLoadConfigFn(
-		[this](FOperationConfig& OperationConfig, const TArray<FString>& SectionHeiarchy)
+		[this](FOperationConfig& OperationConfig, const TArray<FString>& SectionHeirarchy)
 		{
-			return LoadConfig(OperationConfig, SectionHeiarchy);
+			return LoadConfig(OperationConfig, SectionHeirarchy);
 		});
 
 	RegisterComponents();
diff --git a/Engine/Plugins/Online/Experimental/OnlineServices/Source/OnlineServicesCommon/Public/Online/OnlineServicesCommon.h b/Engine/Plugins/Online/Experimental/OnlineServices/Source/OnlineServicesCommon/Public/Online/OnlineServicesCommon.h
--- a/Engine/Plugins/Online/Experimental/OnlineServices/Source/OnlineServicesCommon/Public/Online/OnlineServicesCommon.h
+++ b/Engine/Plugins/Online/Experimental/OnlineServices/Source/OnlineServicesCommon/Public/Online/OnlineServicesCommon.h
@@ -193,6 +193,7 @@ public:
 	 * Load a config struct for an interface + operation
 	 * Will load values from the following sections:
 	 *   OnlineServices
+	 *   OnlineServices.<InterfaceName> (if InterfaceName is set)
 	 *   OnlineServices.<ServiceProvider>
 	 *   OnlineServices.<ServiceProvider>.<InterfaceName> (if InterfaceName is set)
 	 *   OnlineServices.<ServiceProvider>.<InterfaceName>.<OperationName> (if OperationName is set)
@@ -209,6 +210,10 @@ public:
 		TArray<FString> SectionHeiarchy;
 		FString SectionName = TEXT("OnlineServices");
 		SectionHeiarchy.Add(SectionName);
+		if (!InterfaceName.IsEmpty())
+		{
+			SectionHeiarchy.Add(SectionName + TEXT(".") + InterfaceName);
+		}
 		SectionName += TEXT(".") + GetConfigName();
 		SectionHeiarchy.Add(SectionName);
 		if (!InterfaceName.IsEmpty())
diff --git a/Engine/Plugins/Online/Experimental/OnlineServices/Source/OnlineServicesInterface/Private/Online/Schema.cpp b/Engine/Plugins/Online/Experimental/OnlineServices/Source/OnlineServicesInterface/Private/Online/Schema.cpp
--- a/Engine/Plugins/Online/Experimental/OnlineServices/Source/OnlineServicesInterface/Private/Online/Schema.cpp
+++ b/Engine/Plugins/Online/Experimental/OnlineServices/Source/OnlineServicesInterface/Private/Online/Schema.cpp
@@ -1,6 +1,7 @@
 // Copyright Epic Games, Inc. All Rights Reserved.
 
 #include "Online/Schema.h"
+
 #include "Algo/Find.h"
 #include "Algo/ForEach.h"
 #include "Algo/Reverse.h"
@@ -52,20 +53,19 @@ bool FSchemaRegistry::ParseConfig(const FSchemaRegistryDescriptorConfig& Config)
 {
 	bool ParsedSuccessfully = true;
 
-	// Populate known schema.
-	// Validation will be handled later due to the dependency that schema has on service descriptors.
-	TMap<FSchemaId, const FSchemaDescriptor*> KnownSchemaDescriptors;
-	for (const FSchemaDescriptor& SchemaDescriptor : Config.SchemaDescriptors)
+	// Populate known service attributes
+	TMap<FSchemaServiceAttributeId, const FSchemaServiceAttributeDescriptor*> KnownSchemaServiceAttributeDescriptors;
+	for (const FSchemaServiceAttributeDescriptor& SchemaServiceAttributeDescriptor : Config.ServiceAttributeDescriptors)
 	{
 		// Check that schema descriptor id has not already been used.
-		if (KnownSchemaDescriptors.Find(SchemaDescriptor.Id))
+		if (KnownSchemaServiceAttributeDescriptors.Contains(SchemaServiceAttributeDescriptor.Id))
 		{
-			UE_LOG(LogOnlineSchema, Error, TEXT("Duplicate schema descriptor found: %s"), *SchemaDescriptor.Id.ToString().ToLower());
+			UE_LOG(LogOnlineSchema, Error, TEXT("Duplicate schema service attribute descriptor found: %s"), *SchemaServiceAttributeDescriptor.Id.ToString().ToLower());
 			ParsedSuccessfully = false;
 			continue;
 		}
 
-		KnownSchemaDescriptors.Add(SchemaDescriptor.Id, &SchemaDescriptor);
+		KnownSchemaServiceAttributeDescriptors.Add(SchemaServiceAttributeDescriptor.Id, &SchemaServiceAttributeDescriptor);
 	}
 
 	// Verify and populate known service descriptors.
@@ -73,7 +73,7 @@ bool FSchemaRegistry::ParseConfig(const FSchemaRegistryDescriptorConfig& Config)
 	for (const FSchemaServiceDescriptor& SchemaServiceDescriptor : Config.ServiceDescriptors)
 	{
 		// Check that schema service descriptor id has not already been used.
-		if (KnownSchemaServiceDescriptors.Find(SchemaServiceDescriptor.Id))
+		if (KnownSchemaServiceDescriptors.Contains(SchemaServiceDescriptor.Id))
 		{
 			UE_LOG(LogOnlineSchema, Error, TEXT("Duplicate schema service descriptor found: %s"), *SchemaServiceDescriptor.Id.ToString().ToLower());
 			ParsedSuccessfully = false;
@@ -82,19 +82,29 @@ bool FSchemaRegistry::ParseConfig(const FSchemaRegistryDescriptorConfig& Config)
 
 		// Verify attributes
 		TSet<FSchemaServiceAttributeId> SeenAttributes;
-		for (const FSchemaServiceAttributeDescriptor& SchemaServiceAttributeDescriptor : SchemaServiceDescriptor.Attributes)
+		for (const FSchemaServiceAttributeId& SchemaServiceAttributeId : SchemaServiceDescriptor.AttributeIds)
 		{
 			// Check that attribute ids are not reused.
-			if (SeenAttributes.Find(SchemaServiceAttributeDescriptor.Id))
+			if (SeenAttributes.Contains(SchemaServiceAttributeId))
 			{
 				UE_LOG(LogOnlineSchema, Error, TEXT("Invalid schema service attribute %s.%s: Attribute id has already been used."),
 					*SchemaServiceDescriptor.Id.ToString().ToLower(),
-					*SchemaServiceAttributeDescriptor.Id.ToString().ToLower());
+					*SchemaServiceAttributeId.ToString().ToLower());
 				ParsedSuccessfully = false;
 				continue;
 			}
-			SeenAttributes.Add(SchemaServiceAttributeDescriptor.Id);
+			SeenAttributes.Add(SchemaServiceAttributeId);
 
+			const FSchemaServiceAttributeDescriptor** SchemaServiceAttributeDescriptorPtr = KnownSchemaServiceAttributeDescriptors.Find(SchemaServiceAttributeId);
+			if (!SchemaServiceAttributeDescriptorPtr)
+			{
+				UE_LOG(LogOnlineSchema, Error, TEXT("Invalid schema service attribute %s.%s: Attribute does not exist."),
+					*SchemaServiceDescriptor.Id.ToString().ToLower(),
+					*SchemaServiceAttributeId.ToString().ToLower());
+				ParsedSuccessfully = false;
+				continue;
+			}
+			const FSchemaServiceAttributeDescriptor& SchemaServiceAttributeDescriptor = **SchemaServiceAttributeDescriptorPtr;
 			const ESchemaServiceAttributeSupportedTypeFlags SupportedTypes = BuildFlagsFromArray(SchemaServiceAttributeDescriptor.SupportedTypes);
 
 			// Check that supported type has been set.
@@ -129,118 +139,305 @@ bool FSchemaRegistry::ParseConfig(const FSchemaRegistryDescriptorConfig& Config)
 		KnownSchemaServiceDescriptors.Add(SchemaServiceDescriptor.Id, &SchemaServiceDescriptor);
 	}
 
-	// Verify schema data.
+	// Populate known schema categories
+	TMap<FSchemaCategoryId, const FSchemaCategoryDescriptor*> KnownSchemaCategoryDescriptors;
+	for (const FSchemaCategoryDescriptor& SchemaCategoryDescriptor : Config.SchemaCategoryDescriptors)
+	{
+		// Check that schema descriptor id has not already been used.
+		if (KnownSchemaCategoryDescriptors.Contains(SchemaCategoryDescriptor.Id))
+		{
+			UE_LOG(LogOnlineSchema, Error, TEXT("Duplicate schema category descriptor found: %s"), *SchemaCategoryDescriptor.Id.ToString().ToLower());
+			ParsedSuccessfully = false;
+			continue;
+		}
+		KnownSchemaCategoryDescriptors.Add(SchemaCategoryDescriptor.Id, &SchemaCategoryDescriptor);
+
+		// Check that schema service descriptor exists
+		if (SchemaCategoryDescriptor.ServiceDescriptorId != FSchemaServiceDescriptorId()
+			&& !KnownSchemaServiceDescriptors.Contains(SchemaCategoryDescriptor.ServiceDescriptorId))
+		{
+			UE_LOG(LogOnlineSchema, Error, TEXT("Invalid schema category %s: Service descriptor id %s not found."),
+				*SchemaCategoryDescriptor.Id.ToString().ToLower(),
+				*SchemaCategoryDescriptor.ServiceDescriptorId.ToString().ToLower());
+			ParsedSuccessfully = false;
+		}
+	}
+
+	// Populate known schema. Verify what can be verified statically. Heirarchy validation is performed below.
+	TMap<FSchemaId, const FSchemaDescriptor*> KnownSchemaDescriptors;
+	for (const FSchemaDescriptor& SchemaDescriptor : Config.SchemaDescriptors)
+	{
+		// Check that schema descriptor id has not already been used.
+		if (KnownSchemaDescriptors.Contains(SchemaDescriptor.Id))
+		{
+			UE_LOG(LogOnlineSchema, Error, TEXT("Duplicate schema descriptor found: %s"), *SchemaDescriptor.Id.ToString().ToLower());
+			ParsedSuccessfully = false;
+			continue;
+		}
+		KnownSchemaDescriptors.Add(SchemaDescriptor.Id, &SchemaDescriptor);
+
+		// Parent must declare at least one category
+		if (SchemaDescriptor.ParentId == FSchemaId() && SchemaDescriptor.CategoryIds.IsEmpty())
+		{
+			// Error, base type must define at least one category
+			UE_LOG(LogOnlineSchema, Error, TEXT("Invalid schema %s: Base schema must define at least one category."),
+				*SchemaDescriptor.Id.ToString().ToLower());
+			ParsedSuccessfully = false;
+		}
+
+		// Children must not declare categories
+		if (SchemaDescriptor.ParentId != FSchemaId() && !SchemaDescriptor.CategoryIds.IsEmpty())
+		{
+			// Error child cannot define categories
+			UE_LOG(LogOnlineSchema, Error, TEXT("Invalid schema %s: Child schema cannot define categories."),
+				*SchemaDescriptor.Id.ToString().ToLower());
+			ParsedSuccessfully = false;
+		}
+
+		TSet<FSchemaCategoryId> SeenCategoryIds;
+		for (const FSchemaCategoryId& CategoryId : SchemaDescriptor.CategoryIds)
+		{
+			if (SeenCategoryIds.Contains(CategoryId))
+			{
+				UE_LOG(LogOnlineSchema, Error, TEXT("Invalid schema category %s.%s: Category id has already been used."),
+					*SchemaDescriptor.Id.ToString().ToLower(),
+					*CategoryId.ToString().ToLower());
+				ParsedSuccessfully = false;
+				continue;
+			}
+			SeenCategoryIds.Emplace(CategoryId);
+
+			if (!KnownSchemaCategoryDescriptors.Contains(CategoryId))
+			{
+				UE_LOG(LogOnlineSchema, Error, TEXT("Invalid schema category %s.%s: Category does not exist."),
+					*SchemaDescriptor.Id.ToString().ToLower(),
+					*CategoryId.ToString().ToLower());
+				ParsedSuccessfully = false;
+			}
+		}
+	}
+
+	// Populate known schema attributes
+	TMap<FSchemaAttributeId, const FSchemaAttributeDescriptor*> KnownSchemaAttributeDescriptors;
+	for (const FSchemaAttributeDescriptor& SchemaAttributeDescriptor : Config.SchemaAttributeDescriptors)
+	{
+		// Check that schema descriptor id has not already been used.
+		if (KnownSchemaAttributeDescriptors.Contains(SchemaAttributeDescriptor.Id))
+		{
+			UE_LOG(LogOnlineSchema, Error, TEXT("Duplicate schema attribute descriptor found: %s"), *SchemaAttributeDescriptor.Id.ToString().ToLower());
+			ParsedSuccessfully = false;
+			continue;
+		}
+		KnownSchemaAttributeDescriptors.Emplace(SchemaAttributeDescriptor.Id, &SchemaAttributeDescriptor);
+		
+		// Check that a type has been set
+		if (SchemaAttributeDescriptor.Type == ESchemaAttributeType::None)
+		{
+			UE_LOG(LogOnlineSchema, Error, TEXT("Invalid schema attribute %s: A valid type must be set."),
+				*SchemaAttributeDescriptor.Id.ToString().ToLower());
+			ParsedSuccessfully = false;
+		}
+
+		// Check that variable length data has a set maximum length.
+		if (SchemaAttributeDescriptor.Type == ESchemaAttributeType::String && SchemaAttributeDescriptor.MaxSize <= 0)
+		{
+			UE_LOG(LogOnlineSchema, Error, TEXT("Invalid schema attribute %s: A valid max size must be set for variable length data."),
+				*SchemaAttributeDescriptor.Id.ToString().ToLower());
+			ParsedSuccessfully = false;
+		}
+
+		// Check that group id is valid.
+		if (SchemaAttributeDescriptor.UpdateGroupId < 0)
+		{
+			UE_LOG(LogOnlineSchema, Error, TEXT("Invalid schema attribute %s: UpdateGroupId must be set to either 0 or a positive integer value."),
+				*SchemaAttributeDescriptor.Id.ToString().ToLower());
+			ParsedSuccessfully = false;
+		}
+		
+		const ESchemaAttributeFlags AttributeFlags = BuildFlagsFromArray(SchemaAttributeDescriptor.Flags);
+
+		if (EnumHasAllFlags(AttributeFlags, ESchemaAttributeFlags::SchemaCompatibilityId))
+		{
+			if (SchemaAttributeDescriptor.Type != ESchemaAttributeType::Int64)
+			{
+				UE_LOG(LogOnlineSchema, Error, TEXT("Invalid schema attribute %s: SchemaCompatibilityId field may only be set to Int64."),
+					*SchemaAttributeDescriptor.Id.ToString().ToLower());
+				ParsedSuccessfully = false;
+			}
+		
+			if (EnumHasAnyFlags(AttributeFlags, ESchemaAttributeFlags::Private))
+			{
+				UE_LOG(LogOnlineSchema, Error, TEXT("Invalid schema attribute %s: SchemaCompatibilityId attribute may not be private."),
+					*SchemaAttributeDescriptor.Id.ToString().ToLower());
+				ParsedSuccessfully = false;
+			}
+		}
+		
+		// Check that public or private visibility has been set and is valid.
+		if (EnumHasAllFlags(AttributeFlags, ESchemaAttributeFlags::Public | ESchemaAttributeFlags::Private))
+		{
+			UE_LOG(LogOnlineSchema, Error, TEXT("Invalid schema attribute %s: Public and Private visibility are mutually exclusive."),
+				*SchemaAttributeDescriptor.Id.ToString().ToLower());
+			ParsedSuccessfully = false;
+		}
+		if (!EnumHasAnyFlags(AttributeFlags, ESchemaAttributeFlags::Public | ESchemaAttributeFlags::Private))
+		{
+			UE_LOG(LogOnlineSchema, Error, TEXT("Invalid schema attribute %s: Either Public or Private visibility must be set."),
+				*SchemaAttributeDescriptor.Id.ToString().ToLower());
+			ParsedSuccessfully = false;
+		}
+		if (EnumHasAllFlags(AttributeFlags, ESchemaAttributeFlags::Searchable | ESchemaAttributeFlags::Private))
+		{
+			UE_LOG(LogOnlineSchema, Error, TEXT("Invalid schema attribute %s: Searchable attributes may not be private."),
+				*SchemaAttributeDescriptor.Id.ToString().ToLower());
+			ParsedSuccessfully = false;
+		}
+	}
+
+	// Populate known schema category attributes 
+	using FCategoryToAttributesMap = TMap<FSchemaCategoryId, TArray<FSchemaAttributeId>>;
+	using FSchemaToCategoryAttributesMap = TMap<FSchemaId, FCategoryToAttributesMap>;
+	FSchemaToCategoryAttributesMap KnownSchemaCategoryAttributes;
+	for (const FSchemaCategoryAttributesDescriptor& SchemaCategoryAttributesDescriptor : Config.SchemaCategoryAttributeDescriptors)
+	{
+		// Check that the id has not already been used.
+		FCategoryToAttributesMap* FoundSchema = KnownSchemaCategoryAttributes.Find(SchemaCategoryAttributesDescriptor.SchemaId);
+		if (FoundSchema && FoundSchema->Contains(SchemaCategoryAttributesDescriptor.CategoryId))
+		{
+			UE_LOG(LogOnlineSchema, Error, TEXT("Duplicate schema additional category attributes descriptor found: %s.%s"),
+				*SchemaCategoryAttributesDescriptor.SchemaId.ToString().ToLower(),
+				*SchemaCategoryAttributesDescriptor.CategoryId.ToString().ToLower());
+			ParsedSuccessfully = false;
+			continue;
+		}
+
+		// Check that the schema exists
+		if (!KnownSchemaDescriptors.Contains(SchemaCategoryAttributesDescriptor.SchemaId))
+		{
+			UE_LOG(LogOnlineSchema, Error, TEXT("Invalid additional category attributes %s.%s: Schema does not exist"),
+				*SchemaCategoryAttributesDescriptor.SchemaId.ToString().ToLower(),
+				*SchemaCategoryAttributesDescriptor.CategoryId.ToString().ToLower());
+			ParsedSuccessfully = false;
+			continue;
+		}
+
+		// Check that the category exists
+		if (!KnownSchemaCategoryDescriptors.Contains(SchemaCategoryAttributesDescriptor.CategoryId))
+		{
+			UE_LOG(LogOnlineSchema, Error, TEXT("Invalid additional category attributes %s.%s: Category does not exist"),
+				*SchemaCategoryAttributesDescriptor.SchemaId.ToString().ToLower(),
+				*SchemaCategoryAttributesDescriptor.CategoryId.ToString().ToLower());
+			ParsedSuccessfully = false;
+			continue;
+		}
+
+		bool bAttributesParsedSuccessfully = true;
+		TSet<FSchemaAttributeId> SeenAttributeIds;
+		for (const FSchemaAttributeId& AttributeId : SchemaCategoryAttributesDescriptor.AttributeIds)
+		{
+			if (SeenAttributeIds.Contains(AttributeId))
+			{
+				UE_LOG(LogOnlineSchema, Error, TEXT("Invalid category attributes %s.%s: Attribute %s duplicated"),
+					*SchemaCategoryAttributesDescriptor.SchemaId.ToString().ToLower(),
+					*SchemaCategoryAttributesDescriptor.CategoryId.ToString().ToLower(),
+					*AttributeId.ToString().ToLower());
+				bAttributesParsedSuccessfully = false;
+				continue;
+			}
+			SeenAttributeIds.Emplace(AttributeId);
+
+			if (!KnownSchemaAttributeDescriptors.Contains(AttributeId))
+			{
+				UE_LOG(LogOnlineSchema, Error, TEXT("Invalid category attributes %s.%s: Attribute %s does not exist"),
+					*SchemaCategoryAttributesDescriptor.SchemaId.ToString().ToLower(),
+					*SchemaCategoryAttributesDescriptor.CategoryId.ToString().ToLower(),
+					*AttributeId.ToString().ToLower());
+				bAttributesParsedSuccessfully = false;
+				continue;
+			}
+		}
+		if (!bAttributesParsedSuccessfully)
+		{
+			ParsedSuccessfully = false;
+			continue;
+		}
+
+		if (!FoundSchema)
+		{
+			FoundSchema = &KnownSchemaCategoryAttributes.Emplace(SchemaCategoryAttributesDescriptor.SchemaId);
+		}
+		FoundSchema->Emplace(SchemaCategoryAttributesDescriptor.CategoryId, SchemaCategoryAttributesDescriptor.AttributeIds);
+	}
+
+	// Verify schema data
 	for (const TPair<FSchemaId, const FSchemaDescriptor*>& SchemaDescriptorPair : KnownSchemaDescriptors)
 	{
-		const FSchemaDescriptor& SchemaDescriptor = *SchemaDescriptorPair.Get<1>();
+		const FSchemaDescriptor& SchemaDescriptorToVerify = *SchemaDescriptorPair.Value;
+
+		// Stack of schemas, will push in tail->root order and then reverse to root->tail order.
+		TArray<FSchemaId> SchemaHierarchy;
 
-		// Verify parent schema.
-		if (SchemaDescriptor.ParentId.IsValid())
+		// Verify hierarchy.
+		const FSchemaDescriptor* TestSchemaDescriptor = &SchemaDescriptorToVerify;
+		while (TestSchemaDescriptor)
 		{
-			TSet<FSchemaId> SeenSchema;
-			const FSchemaDescriptor* BaseParentSchemaDescriptor = nullptr;
-			const FSchemaDescriptor* TestSchemaDescriptor = &SchemaDescriptor;
-			while (TestSchemaDescriptor)
+			// Check for circular dependencies in schema.
+			if (SchemaHierarchy.Contains(TestSchemaDescriptor->Id))
 			{
-				// Check for circular dependencies in schema.
-				if (SeenSchema.Find(TestSchemaDescriptor->Id))
-				{
-					UE_LOG(LogOnlineSchema, Error, TEXT("Circular parent dependency found in schema: %s"), *SchemaDescriptor.Id.ToString().ToLower());
-					ParsedSuccessfully = false;
-					break;
-				}
-				SeenSchema.Add(TestSchemaDescriptor->Id);
+				UE_LOG(LogOnlineSchema, Error, TEXT("Circular parent dependency found in schema: %s"), *SchemaDescriptorToVerify.Id.ToString().ToLower());
+				ParsedSuccessfully = false;
+				break;
+			}
+			SchemaHierarchy.Push(TestSchemaDescriptor->Id);
 
-				// Iterate parent.
-				if (TestSchemaDescriptor->ParentId != FSchemaId())
+			if (TestSchemaDescriptor->ParentId != FSchemaId())
+			{
+				// Move onto next layer in the hierarchy
+				if (const FSchemaDescriptor** ParentDescriptor = KnownSchemaDescriptors.Find(TestSchemaDescriptor->ParentId))
 				{
-					// Check that parent schema exists.
-					if (const FSchemaDescriptor** FoundDescriptor = KnownSchemaDescriptors.Find(TestSchemaDescriptor->ParentId))
-					{
-						TestSchemaDescriptor = *FoundDescriptor;
-						BaseParentSchemaDescriptor = TestSchemaDescriptor;
-					}
-					else
-					{
-						UE_LOG(LogOnlineSchema, Error, TEXT("Invalid schema %s: Parent schema %s does not exist."),
-							*SchemaDescriptor.Id.ToString().ToLower(),
-							*TestSchemaDescriptor->ParentId.ToString().ToLower());
-						ParsedSuccessfully = false;
-						TestSchemaDescriptor = nullptr;
-					}
+					TestSchemaDescriptor = *ParentDescriptor;
 				}
 				else
 				{
-					TestSchemaDescriptor = nullptr;
+					UE_LOG(LogOnlineSchema, Error, TEXT("Invalid schema %s: Parent schema %s does not exist."),
+						*SchemaDescriptorToVerify.Id.ToString().ToLower(),
+						*TestSchemaDescriptor->ParentId.ToString().ToLower());
+					ParsedSuccessfully = false;
+					break;
 				}
 			}
-
-			// verify that all categories used in schema exist in base parent schema.
-			if (BaseParentSchemaDescriptor)
+			else
 			{
-				for (const FSchemaCategoryDescriptor& SchemaCategoryDescriptor : SchemaDescriptor.Categories)
-				{
-					bool bBaseCategoryExists = Algo::FindByPredicate(BaseParentSchemaDescriptor->Categories,
-					[&SchemaCategoryDescriptor](const FSchemaCategoryDescriptor& Descriptor)
-					{
-						return Descriptor.Id == SchemaCategoryDescriptor.Id;
-					}) != nullptr;
-
-					if (!bBaseCategoryExists)
-					{
-						UE_LOG(LogOnlineSchema, Error, TEXT("Invalid schema %s: Category %s does not exist in base parent %s schema."),
-							*SchemaDescriptor.Id.ToString().ToLower(),
-							*SchemaCategoryDescriptor.Id.ToString().ToLower(),
-							*BaseParentSchemaDescriptor->Id.ToString().ToLower());
-							ParsedSuccessfully = false;
-					}
-				}
+				TestSchemaDescriptor = nullptr;
 			}
 		}
 
-		// Verify category descriptors.
-		TSet<FSchemaCategoryId> SeenSchemaCategoryIds;
+		// Reverse the stack so it's in root->tail order
+		Algo::Reverse(SchemaHierarchy);
+		check(SchemaHierarchy.Num() > 0);
+		const FSchemaId RootSchemaId = SchemaHierarchy[0];
+		const FSchemaDescriptor& RootSchemaDescriptor = *KnownSchemaDescriptors.FindChecked(RootSchemaId);
+
+		// Verify categories
 		TSet<FSchemaServiceDescriptorId> SeenSchemaServiceDescriptorIds;
-		for (const FSchemaCategoryDescriptor& SchemaCategoryDescriptor : SchemaDescriptor.Categories)
+		for (const FSchemaCategoryId& SchemaCategoryId : RootSchemaDescriptor.CategoryIds)
 		{
-			if (SeenSchemaCategoryIds.Find(SchemaCategoryDescriptor.Id))
+			const FSchemaCategoryDescriptor** SchemaCategoryDescriptorPtr = KnownSchemaCategoryDescriptors.Find(SchemaCategoryId);
+			if (!SchemaCategoryDescriptorPtr)
 			{
-				UE_LOG(LogOnlineSchema, Error, TEXT("Invalid schema category %s.%s: Category id has already been used."),
-					*SchemaDescriptor.Id.ToString().ToLower(),
-					*SchemaCategoryDescriptor.Id.ToString().ToLower());
-				ParsedSuccessfully = false;
+				// Already reported in schema population checks above.
 				continue;
 			}
-			SeenSchemaCategoryIds.Add(SchemaCategoryDescriptor.Id);
+			const FSchemaCategoryDescriptor& SchemaCategoryDescriptor = **SchemaCategoryDescriptorPtr;
 
 			// Verify that service descriptor is valid.
 			if (SchemaCategoryDescriptor.ServiceDescriptorId != FSchemaServiceDescriptorId())
 			{
-				// Verify that only a base parent set service descriptor.
-				if (SchemaDescriptor.ParentId != FSchemaId())
-				{
-					UE_LOG(LogOnlineSchema, Error, TEXT("Invalid schema category %s.%s: Service descriptor id may only be set by a base parent schmea category."),
-						*SchemaDescriptor.Id.ToString().ToLower(),
-						*SchemaCategoryDescriptor.Id.ToString().ToLower(),
-						*SchemaCategoryDescriptor.ServiceDescriptorId.ToString().ToLower());
-					ParsedSuccessfully = false;
-				}
-
-				// Verify that service descriptor exists.
-				if (KnownSchemaServiceDescriptors.Find(SchemaCategoryDescriptor.ServiceDescriptorId) == nullptr)
-				{
-					UE_LOG(LogOnlineSchema, Error, TEXT("Invalid schema category %s.%s: Service descriptor id %s not found."),
-						*SchemaDescriptor.Id.ToString().ToLower(),
-						*SchemaCategoryDescriptor.Id.ToString().ToLower(),
-						*SchemaCategoryDescriptor.ServiceDescriptorId.ToString().ToLower());
-					ParsedSuccessfully = false;
-				}
-
 				// Verify that the service descriptor is not reused between categories.
 				if (SeenSchemaServiceDescriptorIds.Find(SchemaCategoryDescriptor.ServiceDescriptorId) != nullptr)
 				{
 					UE_LOG(LogOnlineSchema, Error, TEXT("Invalid schema category %s.%s: Service descriptor id %s has already been used by another category."),
-						*SchemaDescriptor.Id.ToString().ToLower(),
+						*SchemaDescriptorToVerify.Id.ToString().ToLower(),
 						*SchemaCategoryDescriptor.Id.ToString().ToLower(),
 						*SchemaCategoryDescriptor.ServiceDescriptorId.ToString().ToLower());
 					ParsedSuccessfully = false;
@@ -252,131 +449,77 @@ bool FSchemaRegistry::ParseConfig(const FSchemaRegistryDescriptorConfig& Config)
 			// Verify category attributes.
 			TSet<FSchemaServiceAttributeId> SeenAttributes;
 			FSchemaAttributeId SeenSchemaCompatibilityAttributeId;
-			for (const FSchemaAttributeDescriptor& SchemaAttributeDescriptor : SchemaCategoryDescriptor.Attributes)
+			TSet<FSchemaAttributeId> SchemaCategoryAttributeIds;
+			// Iterate down the stack, checking attributes as we go
+			for (const FSchemaId SchemaHierarchyLayerId : SchemaHierarchy)
 			{
-				// Check that attribute ids are not reused.
-				if (SeenAttributes.Find(SchemaAttributeDescriptor.Id))
+				const FCategoryToAttributesMap* FoundSchemaCategories = KnownSchemaCategoryAttributes.Find(SchemaHierarchyLayerId);
+				const TArray<FSchemaAttributeId>* FoundCategoryAttributes = FoundSchemaCategories ? FoundSchemaCategories->Find(SchemaCategoryId) : nullptr;
+				if (!FoundCategoryAttributes)
 				{
-					UE_LOG(LogOnlineSchema, Error, TEXT("Invalid schema attribute %s.%s.%s: Attribute id has already been used."),
-						*SchemaDescriptor.Id.ToString().ToLower(),
-						*SchemaCategoryDescriptor.Id.ToString().ToLower(),
-						*SchemaAttributeDescriptor.Id.ToString().ToLower());
-					ParsedSuccessfully = false;
+					// This layer doesn't add any attribs for this category.
 					continue;
 				}
-				SeenAttributes.Add(SchemaAttributeDescriptor.Id);
 
-				// Check that a type has been set
-				if (SchemaAttributeDescriptor.Type == ESchemaAttributeType::None)
+				for (const FSchemaAttributeId& SchemaAttributeId : *FoundCategoryAttributes)
 				{
-					UE_LOG(LogOnlineSchema, Error, TEXT("Invalid schema attribute %s.%s.%s: A valid type must be set."),
-						*SchemaDescriptor.Id.ToString().ToLower(),
-						*SchemaCategoryDescriptor.Id.ToString().ToLower(),
-						*SchemaAttributeDescriptor.Id.ToString().ToLower());
-					ParsedSuccessfully = false;
-				}
-
-				// Check that variable length data has a set maximum length.
-				if (SchemaAttributeDescriptor.Type == ESchemaAttributeType::String && SchemaAttributeDescriptor.MaxSize <= 0)
-				{
-					UE_LOG(LogOnlineSchema, Error, TEXT("Invalid schema attribute %s.%s.%s: A valid max size must be set for variable length data."),
-						*SchemaDescriptor.Id.ToString().ToLower(),
-						*SchemaCategoryDescriptor.Id.ToString().ToLower(),
-						*SchemaAttributeDescriptor.Id.ToString().ToLower());
-					ParsedSuccessfully = false;
-				}
-
-				const ESchemaAttributeFlags AttributeFlags = BuildFlagsFromArray(SchemaAttributeDescriptor.Flags);
-
-				// Check that searchable attributes only exist in a base schema.
-				if (EnumHasAllFlags(AttributeFlags, ESchemaAttributeFlags::Searchable) && SchemaDescriptor.ParentId != FSchemaId())
-				{
-					UE_LOG(LogOnlineSchema, Error, TEXT("Invalid schema attribute %s.%s.%s: Searchable fields may only exist in the base schema."),
-						*SchemaDescriptor.Id.ToString().ToLower(),
-						*SchemaCategoryDescriptor.Id.ToString().ToLower(),
-						*SchemaAttributeDescriptor.Id.ToString().ToLower());
-					ParsedSuccessfully = false;
-				}
-
-				// Check that public or private visibility has been set and is valid.
-				if (EnumHasAllFlags(AttributeFlags, ESchemaAttributeFlags::Public | ESchemaAttributeFlags::Private))
-				{
-					UE_LOG(LogOnlineSchema, Error, TEXT("Invalid schema attribute %s.%s.%s: Public and Private visibility are mutually exclusive."),
-						*SchemaDescriptor.Id.ToString().ToLower(),
-						*SchemaCategoryDescriptor.Id.ToString().ToLower(),
-						*SchemaAttributeDescriptor.Id.ToString().ToLower());
-					ParsedSuccessfully = false;
-				}
-				if (!EnumHasAnyFlags(AttributeFlags, ESchemaAttributeFlags::Public | ESchemaAttributeFlags::Private))
-				{
-					UE_LOG(LogOnlineSchema, Error, TEXT("Invalid schema attribute %s.%s.%s: Either Public or Private visibility must be set."),
-						*SchemaDescriptor.Id.ToString().ToLower(),
-						*SchemaCategoryDescriptor.Id.ToString().ToLower(),
-						*SchemaAttributeDescriptor.Id.ToString().ToLower());
-					ParsedSuccessfully = false;
-				}
-				if (EnumHasAllFlags(AttributeFlags, ESchemaAttributeFlags::Searchable | ESchemaAttributeFlags::Private))
-				{
-					UE_LOG(LogOnlineSchema, Error, TEXT("Invalid schema attribute %s.%s.%s: Searchable attributes may not be private."),
-						*SchemaDescriptor.Id.ToString().ToLower(),
-						*SchemaCategoryDescriptor.Id.ToString().ToLower(),
-						*SchemaAttributeDescriptor.Id.ToString().ToLower());
-					ParsedSuccessfully = false;
-				}
-
-				// Check that schema compatibility attribute only exists in a base schema and used only once and that it is the Int64 type.
-				if (EnumHasAllFlags(AttributeFlags, ESchemaAttributeFlags::SchemaCompatibilityId))
-				{
-					if (SchemaDescriptor.ParentId != FSchemaId())
+					// Check that attribute ids are not reused.
+					if (SeenAttributes.Find(SchemaAttributeId))
 					{
-						UE_LOG(LogOnlineSchema, Error, TEXT("Invalid schema attribute %s.%s.%s: SchemaCompatibilityId field may only exist in the base schema."),
-							*SchemaDescriptor.Id.ToString().ToLower(),
+						UE_LOG(LogOnlineSchema, Error, TEXT("Invalid schema attribute %s.%s.%s: Attribute id has already been used."),
+							*SchemaDescriptorToVerify.Id.ToString().ToLower(),
 							*SchemaCategoryDescriptor.Id.ToString().ToLower(),
-							*SchemaAttributeDescriptor.Id.ToString().ToLower());
+							*SchemaAttributeId.ToString().ToLower());
 						ParsedSuccessfully = false;
+						continue;
 					}
+					SeenAttributes.Add(SchemaAttributeId);
 
-					if (SeenSchemaCompatibilityAttributeId != FSchemaId())
-					{
-						UE_LOG(LogOnlineSchema, Error, TEXT("Invalid schema attribute %s.%s.%s: SchemaCompatibilityId field has already been used by attribute %s."),
-							*SchemaDescriptor.Id.ToString().ToLower(),
-							*SchemaCategoryDescriptor.Id.ToString().ToLower(),
-							*SchemaAttributeDescriptor.Id.ToString().ToLower(),
-							*SeenSchemaCompatibilityAttributeId.ToString().ToLower());
-						ParsedSuccessfully = false;
-					}
-					else
+					const FSchemaAttributeDescriptor** SchemaAttributeDescriptorPtr = KnownSchemaAttributeDescriptors.Find(SchemaAttributeId);
+					if (!SchemaAttributeDescriptorPtr)
 					{
-						SeenSchemaCompatibilityAttributeId = SchemaAttributeDescriptor.Id;
+						// Already reported in the category attribute population checks above.
+						continue;
 					}
+					const FSchemaAttributeDescriptor& SchemaAttributeDescriptor = **SchemaAttributeDescriptorPtr;
+					const ESchemaAttributeFlags AttributeFlags = BuildFlagsFromArray(SchemaAttributeDescriptor.Flags);
 
-					if (SchemaAttributeDescriptor.Type != ESchemaAttributeType::Int64)
+					// Check that searchable attributes only exist in a base schema.
+					if (EnumHasAllFlags(AttributeFlags, ESchemaAttributeFlags::Searchable) && SchemaHierarchyLayerId != RootSchemaId)
 					{
-						UE_LOG(LogOnlineSchema, Error, TEXT("Invalid schema attribute %s.%s.%s: SchemaCompatibilityId field may only be set to Int64."),
-							*SchemaDescriptor.Id.ToString().ToLower(),
+						UE_LOG(LogOnlineSchema, Error, TEXT("Invalid schema attribute %s.%s.%s: Searchable fields may only exist in the base schema."),
+							*SchemaDescriptorToVerify.Id.ToString().ToLower(),
 							*SchemaCategoryDescriptor.Id.ToString().ToLower(),
 							*SchemaAttributeDescriptor.Id.ToString().ToLower());
 						ParsedSuccessfully = false;
 					}
 
-					if (EnumHasAnyFlags(AttributeFlags, ESchemaAttributeFlags::Private))
+					// Check that schema compatibility attribute only exists in a base schema and used only once and that it is the Int64 type.
+					if (EnumHasAllFlags(AttributeFlags, ESchemaAttributeFlags::SchemaCompatibilityId))
 					{
-						UE_LOG(LogOnlineSchema, Error, TEXT("Invalid schema attribute %s.%s.%s: SchemaCompatibilityId attribute may not be private."),
-							*SchemaDescriptor.Id.ToString().ToLower(),
-							*SchemaCategoryDescriptor.Id.ToString().ToLower(),
-							*SchemaAttributeDescriptor.Id.ToString().ToLower());
-						ParsedSuccessfully = false;
-					}
-				}
+						if (SchemaHierarchyLayerId != RootSchemaId)
+						{
+							UE_LOG(LogOnlineSchema, Error, TEXT("Invalid schema attribute %s.%s.%s: SchemaCompatibilityId field may only exist in the base schema."),
+								*SchemaDescriptorToVerify.Id.ToString().ToLower(),
+								*SchemaCategoryDescriptor.Id.ToString().ToLower(),
+								*SchemaAttributeDescriptor.Id.ToString().ToLower());
+							ParsedSuccessfully = false;
+						}
 
-				// Check that group id is valid.
-				if (SchemaAttributeDescriptor.UpdateGroupId < 0)
-				{
-					UE_LOG(LogOnlineSchema, Error, TEXT("Invalid schema attribute %s.%s.%s: UpdateGroupId must be set to either 0 or a positive integer value."),
-						*SchemaDescriptor.Id.ToString().ToLower(),
-						*SchemaCategoryDescriptor.Id.ToString().ToLower(),
-						*SchemaAttributeDescriptor.Id.ToString().ToLower());
-					ParsedSuccessfully = false;
+						if (SeenSchemaCompatibilityAttributeId != FSchemaId())
+						{
+							UE_LOG(LogOnlineSchema, Error, TEXT("Invalid schema attribute %s.%s.%s: SchemaCompatibilityId field has already been used by attribute %s."),
+								*SchemaDescriptorToVerify.Id.ToString().ToLower(),
+								*SchemaCategoryDescriptor.Id.ToString().ToLower(),
+								*SchemaAttributeDescriptor.Id.ToString().ToLower(),
+								*SeenSchemaCompatibilityAttributeId.ToString().ToLower());
+							ParsedSuccessfully = false;
+						}
+						else
+						{
+							SeenSchemaCompatibilityAttributeId = SchemaAttributeDescriptor.Id;
+						}
+					}
 				}
 			}
 		}
@@ -385,7 +528,6 @@ bool FSchemaRegistry::ParseConfig(const FSchemaRegistryDescriptorConfig& Config)
 	// Build definitions only when there are no parsing errors.
 	TMap<FSchemaId, TSharedRef<const FSchemaDefinition>> ParsedSchemaDefinitions;
 	TMap<int64, TSharedRef<const FSchemaDefinition>> ParsedSchemaDefinitionsByCompatibilityId;
-
 	if (ParsedSuccessfully)
 	{
 		// Build schema definitions
@@ -399,166 +541,194 @@ bool FSchemaRegistry::ParseConfig(const FSchemaRegistryDescriptorConfig& Config)
 			ParsedSchemaDefinitions.Add(SchemaDescriptor.Id, SchemaDefinition);
 			SchemaDefinition->Id = SchemaDescriptor.Id;
 
-			TMap<FSchemaCategoryId, FSchemaServiceDescriptorId> ServiceDescriptors;
-			TArray<FSchemaAttributeDefinition*> SchemaAttributeDefinitionAssignmentStorage;
-
-			// Populate definition from descriptor.
-			for (const FSchemaDescriptor** IterSchemaDescriptor = KnownSchemaDescriptors.Find(SchemaDescriptor.Id);
+			// Get the hierarchy from root->leaf
+			TArray<FSchemaId> SchemaHierarchy;
+			SchemaHierarchy.Emplace(SchemaDefinition->Id);
+			for (const FSchemaDescriptor** IterSchemaDescriptor = KnownSchemaDescriptors.Find(SchemaDescriptor.ParentId);
 				IterSchemaDescriptor != nullptr;
 				IterSchemaDescriptor = KnownSchemaDescriptors.Find((*IterSchemaDescriptor)->ParentId))
 			{
-				const FSchemaDescriptor& CurrentSchemaDescriptor = **IterSchemaDescriptor;
+				const FSchemaId& IterSchemaId = (*IterSchemaDescriptor)->Id;
+				SchemaHierarchy.Emplace(IterSchemaId);
 
 				// Populate parent info.
-				if (CurrentSchemaDescriptor.Id != SchemaDescriptor.Id)
-				{
-					SchemaDefinition->ParentSchemaIds.Add(CurrentSchemaDescriptor.Id);
-					SchemaDataCrc = FCrc::TypeCrc32(CurrentSchemaDescriptor.Id, SchemaDataCrc);
-				}
-
-				// Populate attribute definitions.
-				for (const FSchemaCategoryDescriptor& CurrentSchemaCategoryDescriptor : CurrentSchemaDescriptor.Categories)
-				{
-					FSchemaCategoryDefinition& SchemaCategoryDefinition = SchemaDefinition->Categories.FindOrAdd(CurrentSchemaCategoryDescriptor.Id);
-					SchemaCategoryDefinition.Id = CurrentSchemaCategoryDescriptor.Id;
-					SchemaDataCrc = FCrc::TypeCrc32(CurrentSchemaCategoryDescriptor.Id, SchemaDataCrc);
-
-					for (const FSchemaAttributeDescriptor& CurrentSchemaAttributeDescriptor : CurrentSchemaCategoryDescriptor.Attributes)
-					{
-						FSchemaAttributeDefinition& AttributeDefinition = SchemaCategoryDefinition.SchemaAttributeDefinitions.Add(CurrentSchemaAttributeDescriptor.Id);
-						AttributeDefinition.Id = CurrentSchemaAttributeDescriptor.Id;
-						AttributeDefinition.Flags = BuildFlagsFromArray(CurrentSchemaAttributeDescriptor.Flags);
-						AttributeDefinition.Type = CurrentSchemaAttributeDescriptor.Type;
-						AttributeDefinition.MaxSize = CurrentSchemaAttributeDescriptor.MaxSize;
+				SchemaDefinition->ParentSchemaIds.Add(IterSchemaId);
+				SchemaDataCrc = FCrc::TypeCrc32(IterSchemaId, SchemaDataCrc);
+			}
+			Algo::Reverse(SchemaHierarchy);
+			check(SchemaHierarchy.Num() > 0);
+			const FSchemaId RootSchemaId = SchemaHierarchy[0];
+			const FSchemaDescriptor& RootSchemaDescriptor = *KnownSchemaDescriptors.FindChecked(RootSchemaId);
 
-						if (EnumHasAnyFlags(AttributeDefinition.Flags, ESchemaAttributeFlags::SchemaCompatibilityId))
-						{
-							SchemaCategoryDefinition.SchemaCompatibilityAttributeId = AttributeDefinition.Id;
-						}
+			// Add the category definitions.
+			for (const FSchemaCategoryId& CategoryId : RootSchemaDescriptor.CategoryIds)
+			{
+				FSchemaCategoryDefinition& SchemaCategoryDefinition = SchemaDefinition->Categories.Emplace(CategoryId);
+				SchemaCategoryDefinition.Id = CategoryId;
 
-						SchemaDataCrc = FCrc::TypeCrc32(AttributeDefinition.Id, SchemaDataCrc);
-						SchemaDataCrc = FCrc::TypeCrc32(AttributeDefinition.Flags, SchemaDataCrc);
-						SchemaDataCrc = FCrc::TypeCrc32(AttributeDefinition.Type, SchemaDataCrc);
-						SchemaDataCrc = FCrc::TypeCrc32(AttributeDefinition.MaxSize, SchemaDataCrc);
+				SchemaDataCrc = FCrc::TypeCrc32(CategoryId, SchemaDataCrc);
+			}
 
-						SchemaAttributeDefinitionAssignmentStorage.Add(&AttributeDefinition);
-					}
+			// Iterate the hierarchy from root -> leaf
+			for (const FSchemaId& SchemaHierarchyLayerId : SchemaHierarchy)
+			{
+				const FSchemaDescriptor& CurrentSchemaDescriptor = *KnownSchemaDescriptors.FindChecked(SchemaHierarchyLayerId);
 
-					// Store service descriptor for assignment.
-					if (CurrentSchemaCategoryDescriptor.ServiceDescriptorId != FSchemaServiceDescriptorId())
+				// Populate attribute definitions for each category
+				for (const FSchemaCategoryId& CurrentSchemaCategoryId : RootSchemaDescriptor.CategoryIds)
+				{
+					if (const FCategoryToAttributesMap* FoundSchemaCategoryAttributes = KnownSchemaCategoryAttributes.Find(CurrentSchemaDescriptor.Id))
 					{
-						ServiceDescriptors.Add(CurrentSchemaCategoryDescriptor.Id, CurrentSchemaCategoryDescriptor.ServiceDescriptorId);
+						if (const TArray<FSchemaAttributeId>* FoundCategoryAttributeIds = FoundSchemaCategoryAttributes->Find(CurrentSchemaCategoryId))
+						{
+							const FSchemaCategoryDescriptor& CurrentSchemaCategoryDescriptor = *KnownSchemaCategoryDescriptors.FindChecked(CurrentSchemaCategoryId);
+							FSchemaCategoryDefinition& SchemaCategoryDefinition = SchemaDefinition->Categories.FindChecked(CurrentSchemaCategoryDescriptor.Id);
+
+							for (const FSchemaAttributeId& CurrentSchemaAttributeId : *FoundCategoryAttributeIds)
+							{
+								check(!SchemaCategoryDefinition.SchemaAttributeDefinitions.Contains(CurrentSchemaAttributeId));
+
+								const FSchemaAttributeDescriptor& CurrentSchemaAttributeDescriptor = *KnownSchemaAttributeDescriptors.FindChecked(CurrentSchemaAttributeId);
+
+								FSchemaAttributeDefinition& AttributeDefinition = SchemaCategoryDefinition.SchemaAttributeDefinitions.Emplace(CurrentSchemaAttributeDescriptor.Id);
+								AttributeDefinition.Id = CurrentSchemaAttributeDescriptor.Id;
+								AttributeDefinition.Flags = BuildFlagsFromArray(CurrentSchemaAttributeDescriptor.Flags);
+								AttributeDefinition.Type = CurrentSchemaAttributeDescriptor.Type;
+								AttributeDefinition.MaxSize = CurrentSchemaAttributeDescriptor.MaxSize;
+
+								if (EnumHasAnyFlags(AttributeDefinition.Flags, ESchemaAttributeFlags::SchemaCompatibilityId))
+								{
+									SchemaCategoryDefinition.SchemaCompatibilityAttributeId = AttributeDefinition.Id;
+								}
+
+								SchemaDataCrc = FCrc::TypeCrc32(AttributeDefinition.Id, SchemaDataCrc);
+								SchemaDataCrc = FCrc::TypeCrc32(AttributeDefinition.Flags, SchemaDataCrc);
+								SchemaDataCrc = FCrc::TypeCrc32(AttributeDefinition.Type, SchemaDataCrc);
+								SchemaDataCrc = FCrc::TypeCrc32(AttributeDefinition.MaxSize, SchemaDataCrc);
+							}
+						}
 					}
 				}
 			}
 
-			// Attributes must be ordered from parent to derived schema to ensure that service attribute assignment is consistent between derived and parent schemas.
-			Algo::Reverse(SchemaAttributeDefinitionAssignmentStorage);
-
-			// Assign attributes to service attributes.
-			for (TPair<FSchemaCategoryId, FSchemaServiceDescriptorId>& ServiceDescriptor : ServiceDescriptors)
+			// Assign attributes to service attributes for each category
+			for (const FSchemaCategoryId& CategoryId : RootSchemaDescriptor.CategoryIds)
 			{
-				FSchemaCategoryDefinition& SchemaCategoryDefinition = *SchemaDefinition->Categories.Find(ServiceDescriptor.Get<0>());
-				const FSchemaServiceDescriptor& SchemaServiceDescriptor = **KnownSchemaServiceDescriptors.Find(ServiceDescriptor.Get<1>());
+				FSchemaCategoryDefinition& SchemaCategoryDefinition = SchemaDefinition->Categories.FindChecked(CategoryId);
 
-				// todo: pack multiple attributes into a service attribute.
-				// For the time being schema attributes and service attributes have a 1:1 relationship.
-
-				// Service attribute map.
-				TMap<FSchemaServiceAttributeId, const FSchemaServiceAttributeDescriptor*> ServiceAttributesById;
-
-				// Bucket service attributes for assignment
-				TMap<const FSchemaServiceAttributeDescriptor*, ESchemaServiceAttributeSupportedTypeFlags> PublicAttributes;
-				TMap<const FSchemaServiceAttributeDescriptor*, ESchemaServiceAttributeSupportedTypeFlags> PrivateAttributes;
-				TMap<const FSchemaServiceAttributeDescriptor*, ESchemaServiceAttributeSupportedTypeFlags> SearchableAttributes;
-
-				// Build out buckets.
-				for (const FSchemaServiceAttributeDescriptor& SchemaServiceAttributeDescriptor : SchemaServiceDescriptor.Attributes)
+				const FSchemaCategoryDescriptor& SchemaCategoryDescriptor = *KnownSchemaCategoryDescriptors.FindChecked(CategoryId);
+				if (SchemaCategoryDescriptor.ServiceDescriptorId != FSchemaServiceDescriptorId())
 				{
-					ServiceAttributesById.Add(SchemaServiceAttributeDescriptor.Id, &SchemaServiceAttributeDescriptor);
+					const FSchemaServiceDescriptor& SchemaServiceDescriptor = *KnownSchemaServiceDescriptors.FindChecked(SchemaCategoryDescriptor.ServiceDescriptorId);
 
-					const ESchemaServiceAttributeSupportedTypeFlags SupportedTypes = BuildFlagsFromArray(SchemaServiceAttributeDescriptor.SupportedTypes);
-					const ESchemaServiceAttributeFlags ServiceAttributeFlags = BuildFlagsFromArray(SchemaServiceAttributeDescriptor.Flags);
-					if (EnumHasAnyFlags(ServiceAttributeFlags, ESchemaServiceAttributeFlags::Public))
-					{
-						PublicAttributes.Add(&SchemaServiceAttributeDescriptor, SupportedTypes);
-					}
-					if (EnumHasAnyFlags(ServiceAttributeFlags, ESchemaServiceAttributeFlags::Private))
-					{
-						PrivateAttributes.Add(&SchemaServiceAttributeDescriptor, SupportedTypes);
-					}
-					if (EnumHasAnyFlags(ServiceAttributeFlags, ESchemaServiceAttributeFlags::Searchable))
-					{
-						SearchableAttributes.Add(&SchemaServiceAttributeDescriptor, SupportedTypes);
-					}
-				}
+					// todo: pack multiple attributes into a service attribute.
+					// For the time being schema attributes and service attributes have a 1:1 relationship.
 
-				// Assign schema attributes based on buckets.
-				for (FSchemaAttributeDefinition* AttributeDefinition : SchemaAttributeDefinitionAssignmentStorage)
-				{
-					// Find a service attribute which can handle the schema attribute.
-					TMap<const FSchemaServiceAttributeDescriptor*, ESchemaServiceAttributeSupportedTypeFlags>* Bucket = nullptr;
-					if (EnumHasAnyFlags(AttributeDefinition->Flags, ESchemaAttributeFlags::Public))
-					{
-						Bucket = EnumHasAnyFlags(AttributeDefinition->Flags, ESchemaAttributeFlags::Searchable) ? &SearchableAttributes : &PublicAttributes;
-					}
-					else
+					// Bucket service attributes for assignment
+					TMap<const FSchemaServiceAttributeDescriptor*, ESchemaServiceAttributeSupportedTypeFlags> PublicAttributes;
+					TMap<const FSchemaServiceAttributeDescriptor*, ESchemaServiceAttributeSupportedTypeFlags> PrivateAttributes;
+					TMap<const FSchemaServiceAttributeDescriptor*, ESchemaServiceAttributeSupportedTypeFlags> SearchableAttributes;
+
+					// Build out buckets.
+					for (const FSchemaServiceAttributeId& SchemaServiceAttributeId : SchemaServiceDescriptor.AttributeIds)
 					{
-						Bucket = &PrivateAttributes;
+						const FSchemaServiceAttributeDescriptor& SchemaServiceAttributeDescriptor = *KnownSchemaServiceAttributeDescriptors.FindChecked(SchemaServiceAttributeId);
+						const ESchemaServiceAttributeSupportedTypeFlags SupportedTypes = BuildFlagsFromArray(SchemaServiceAttributeDescriptor.SupportedTypes);
+						const ESchemaServiceAttributeFlags ServiceAttributeFlags = BuildFlagsFromArray(SchemaServiceAttributeDescriptor.Flags);
+						if (EnumHasAnyFlags(ServiceAttributeFlags, ESchemaServiceAttributeFlags::Public))
+						{
+							PublicAttributes.Add(&SchemaServiceAttributeDescriptor, SupportedTypes);
+						}
+						if (EnumHasAnyFlags(ServiceAttributeFlags, ESchemaServiceAttributeFlags::Private))
+						{
+							PrivateAttributes.Add(&SchemaServiceAttributeDescriptor, SupportedTypes);
+						}
+						if (EnumHasAnyFlags(ServiceAttributeFlags, ESchemaServiceAttributeFlags::Searchable))
+						{
+							SearchableAttributes.Add(&SchemaServiceAttributeDescriptor, SupportedTypes);
+						}
 					}
 
-					const TPair<const FSchemaServiceAttributeDescriptor*, ESchemaServiceAttributeSupportedTypeFlags>* SchemaServiceAttributeDescriptorPair = Algo::FindByPredicate(*Bucket,
-					[AttributeDefinition](const TPair<const FSchemaServiceAttributeDescriptor*, ESchemaServiceAttributeSupportedTypeFlags>& ServiceAttributePair)
+					// Assign schema attributes based on buckets.
+					for (const FSchemaId& SchemaHierarchyLayerId : SchemaHierarchy)
 					{
-						const FSchemaServiceAttributeDescriptor* SchemaServiceAttributeDescriptor = ServiceAttributePair.Get<0>();
-						const ESchemaServiceAttributeSupportedTypeFlags SchemaServiceAttributeSupportedTypeFlags = ServiceAttributePair.Get<1>();
-						return EnumHasAnyFlags(SchemaServiceAttributeSupportedTypeFlags, TranslateAttributeType(AttributeDefinition->Type)) &&
-							AttributeDefinition->MaxSize <= SchemaServiceAttributeDescriptor->MaxSize;
-					});
-					const FSchemaServiceAttributeDescriptor* FoundServiceAttributeDescriptor = SchemaServiceAttributeDescriptorPair ? SchemaServiceAttributeDescriptorPair->Get<0>() : nullptr;
-
-					if (FoundServiceAttributeDescriptor)
-					{
-						// Todo: Reuse descriptor for assignment until it is full.
-
-						// Remove found descriptor from buckets.
-						PublicAttributes.Remove(FoundServiceAttributeDescriptor);
-						PrivateAttributes.Remove(FoundServiceAttributeDescriptor);
-						SearchableAttributes.Remove(FoundServiceAttributeDescriptor);
-
-						// Create service attribute definition.
-						FSchemaServiceAttributeDefinition& SchemaServiceAttributeDefinition = SchemaCategoryDefinition.ServiceAttributeDefinitions.Add(FoundServiceAttributeDescriptor->Id);
-						SchemaServiceAttributeDefinition.Id = FoundServiceAttributeDescriptor->Id;
-						SchemaServiceAttributeDefinition.Type = TranslateAttributeType(AttributeDefinition->Type);
-						SchemaServiceAttributeDefinition.Flags = TranslateAttributeFlags(AttributeDefinition->Flags);
-						SchemaServiceAttributeDefinition.MaxSize = FoundServiceAttributeDescriptor->MaxSize;
-						SchemaServiceAttributeDefinition.SchemaAttributeIds.Add(AttributeDefinition->Id);
-						AttributeDefinition->ServiceAttributeId = SchemaServiceAttributeDefinition.Id;
-
-						if (EnumHasAnyFlags(AttributeDefinition->Flags, ESchemaAttributeFlags::SchemaCompatibilityId))
+						const FCategoryToAttributesMap* FoundSchemaCategoryAttributes = KnownSchemaCategoryAttributes.Find(SchemaHierarchyLayerId);
+						const TArray<FSchemaAttributeId>* FoundCategoryAttributeIds = FoundSchemaCategoryAttributes ? FoundSchemaCategoryAttributes->Find(CategoryId) : nullptr;
+						if (!FoundCategoryAttributeIds)
 						{
-							SchemaCategoryDefinition.SchemaCompatibilityServiceAttributeId = SchemaServiceAttributeDefinition.Id;
+							continue;
 						}
 
-						// Add service attribute definition to data crc.
-						SchemaDataCrc = FCrc::TypeCrc32(SchemaServiceAttributeDefinition.Id, SchemaDataCrc);
-						SchemaDataCrc = FCrc::TypeCrc32(SchemaServiceAttributeDefinition.Type, SchemaDataCrc);
-						SchemaDataCrc = FCrc::TypeCrc32(SchemaServiceAttributeDefinition.Flags, SchemaDataCrc);
-						SchemaDataCrc = FCrc::TypeCrc32(SchemaServiceAttributeDefinition.MaxSize, SchemaDataCrc);
-						Algo::ForEach(SchemaServiceAttributeDefinition.SchemaAttributeIds, [&SchemaDataCrc](const FSchemaAttributeId& Id){ SchemaDataCrc = FCrc::TypeCrc32(Id, SchemaDataCrc); });
-					}
-					else
-					{
-						UE_LOG(LogOnlineSchema, Error, TEXT("Failed to find service attribute to fit schema attribute %s.%s.%s."),
-							*SchemaDefinition->Id.ToString().ToLower(),
-							*SchemaCategoryDefinition.Id.ToString().ToLower(),
-							*AttributeDefinition->Id.ToString().ToLower());
-						ParsedSuccessfully = false;
+						for (const FSchemaAttributeId& SchemaAttributeId : *FoundCategoryAttributeIds)
+						{
+							FSchemaAttributeDefinition& SchemaAttributeDefinition = SchemaCategoryDefinition.SchemaAttributeDefinitions.FindChecked(SchemaAttributeId);
+
+							// Find a service attribute which can handle the schema attribute.
+							TMap<const FSchemaServiceAttributeDescriptor*, ESchemaServiceAttributeSupportedTypeFlags>* Bucket = nullptr;
+							if (EnumHasAnyFlags(SchemaAttributeDefinition.Flags, ESchemaAttributeFlags::Public))
+							{
+								Bucket = EnumHasAnyFlags(SchemaAttributeDefinition.Flags, ESchemaAttributeFlags::Searchable) ? &SearchableAttributes : &PublicAttributes;
+							}
+							else
+							{
+								Bucket = &PrivateAttributes;
+							}
+
+							const TPair<const FSchemaServiceAttributeDescriptor*, ESchemaServiceAttributeSupportedTypeFlags>* SchemaServiceAttributeDescriptorPair = Algo::FindByPredicate(*Bucket,
+								[SchemaAttributeDefinition](const TPair<const FSchemaServiceAttributeDescriptor*, ESchemaServiceAttributeSupportedTypeFlags>& ServiceAttributePair)
+							{
+								const FSchemaServiceAttributeDescriptor* SchemaServiceAttributeDescriptor = ServiceAttributePair.Get<0>();
+								const ESchemaServiceAttributeSupportedTypeFlags SchemaServiceAttributeSupportedTypeFlags = ServiceAttributePair.Get<1>();
+								return EnumHasAnyFlags(SchemaServiceAttributeSupportedTypeFlags, TranslateAttributeType(SchemaAttributeDefinition.Type)) &&
+									SchemaAttributeDefinition.MaxSize <= SchemaServiceAttributeDescriptor->MaxSize;
+							});
+							const FSchemaServiceAttributeDescriptor* FoundServiceAttributeDescriptor = SchemaServiceAttributeDescriptorPair ? SchemaServiceAttributeDescriptorPair->Get<0>() : nullptr;
+
+							if (FoundServiceAttributeDescriptor)
+							{
+								// Todo: Reuse descriptor for assignment until it is full.
+
+								// Remove found descriptor from buckets.
+								PublicAttributes.Remove(FoundServiceAttributeDescriptor);
+								PrivateAttributes.Remove(FoundServiceAttributeDescriptor);
+								SearchableAttributes.Remove(FoundServiceAttributeDescriptor);
+
+								// Create service attribute definition.
+								FSchemaServiceAttributeDefinition& SchemaServiceAttributeDefinition = SchemaCategoryDefinition.ServiceAttributeDefinitions.Add(FoundServiceAttributeDescriptor->Id);
+								SchemaServiceAttributeDefinition.Id = FoundServiceAttributeDescriptor->Id;
+								SchemaServiceAttributeDefinition.Type = TranslateAttributeType(SchemaAttributeDefinition.Type);
+								SchemaServiceAttributeDefinition.Flags = TranslateAttributeFlags(SchemaAttributeDefinition.Flags);
+								SchemaServiceAttributeDefinition.MaxSize = FoundServiceAttributeDescriptor->MaxSize;
+								SchemaServiceAttributeDefinition.SchemaAttributeIds.Add(SchemaAttributeDefinition.Id);
+								SchemaAttributeDefinition.ServiceAttributeId = SchemaServiceAttributeDefinition.Id;
+
+								if (EnumHasAnyFlags(SchemaAttributeDefinition.Flags, ESchemaAttributeFlags::SchemaCompatibilityId))
+								{
+									SchemaCategoryDefinition.SchemaCompatibilityServiceAttributeId = SchemaServiceAttributeDefinition.Id;
+								}
+
+								// Add service attribute definition to data crc.
+								SchemaDataCrc = FCrc::TypeCrc32(SchemaServiceAttributeDefinition.Id, SchemaDataCrc);
+								SchemaDataCrc = FCrc::TypeCrc32(SchemaServiceAttributeDefinition.Type, SchemaDataCrc);
+								SchemaDataCrc = FCrc::TypeCrc32(SchemaServiceAttributeDefinition.Flags, SchemaDataCrc);
+								SchemaDataCrc = FCrc::TypeCrc32(SchemaServiceAttributeDefinition.MaxSize, SchemaDataCrc);
+
+								Algo::ForEach(SchemaServiceAttributeDefinition.SchemaAttributeIds, [&SchemaDataCrc](const FSchemaAttributeId& Id) { SchemaDataCrc = FCrc::TypeCrc32(Id, SchemaDataCrc); });
+							}
+							else
+							{
+								UE_LOG(LogOnlineSchema, Error, TEXT("Failed to find service attribute to fit schema attribute %s.%s.%s."),
+									*SchemaDefinition->Id.ToString().ToLower(),
+									*SchemaCategoryDefinition.Id.ToString().ToLower(),
+									*SchemaAttributeDefinition.Id.ToString().ToLower());
+								ParsedSuccessfully = false;
+							}
+						}
 					}
-				}
 
-				SchemaDataCrc = FCrc::TypeCrc32(SchemaCategoryDefinition.SchemaCompatibilityAttributeId, SchemaDataCrc);
-				SchemaDataCrc = FCrc::TypeCrc32(SchemaCategoryDefinition.SchemaCompatibilityServiceAttributeId, SchemaDataCrc);
+					SchemaDataCrc = FCrc::TypeCrc32(SchemaCategoryDefinition.SchemaCompatibilityAttributeId, SchemaDataCrc);
+					SchemaDataCrc = FCrc::TypeCrc32(SchemaCategoryDefinition.SchemaCompatibilityServiceAttributeId, SchemaDataCrc);
+				}
 			}
 
 			// Build compatibility id from the schema and data CRCs.
diff --git a/Engine/Plugins/Online/Experimental/OnlineServices/Source/OnlineServicesInterface/Public/Online/SchemaTypes.h b/Engine/Plugins/Online/Experimental/OnlineServices/Source/OnlineServicesInterface/Public/Online/SchemaTypes.h
--- a/Engine/Plugins/Online/Experimental/OnlineServices/Source/OnlineServicesInterface/Public/Online/SchemaTypes.h
+++ b/Engine/Plugins/Online/Experimental/OnlineServices/Source/OnlineServicesInterface/Public/Online/SchemaTypes.h
@@ -284,7 +284,7 @@ struct FSchemaServiceDescriptor
 	/**
 	  * Service attribute descriptions.
 	  */
-	TArray<FSchemaServiceAttributeDescriptor> Attributes;
+	TArray<FSchemaServiceAttributeId> AttributeIds;
 };
 
 /**
@@ -338,10 +338,25 @@ struct FSchemaCategoryDescriptor
 	  * May only be set by base parent schema.
 	  */
 	FSchemaServiceDescriptorId ServiceDescriptorId;
+};
+
+/**
+  * Adds a set of attributes to a given category in a given schema.
+  */
+struct FSchemaCategoryAttributesDescriptor
+{
 	/**
-	  * Attribute descriptors.
-	  */
-	TArray<FSchemaAttributeDescriptor> Attributes;
+	 * The schema to add attributes to.
+	 */
+	FSchemaId SchemaId;
+	/**
+	 * The category to add attributes to.
+	 */
+	FSchemaCategoryId CategoryId;
+	/**
+	 * The attributes to add.
+	 */
+	TArray<FSchemaAttributeId> AttributeIds;
 };
 
 /**
@@ -362,7 +377,7 @@ struct FSchemaDescriptor
 	  * Categories of data within the schema. Each category is handled differently within the
 	  * service so will have different configuration of attributes and service attributes.
 	  */
-	TArray<FSchemaCategoryDescriptor> Categories;
+	TArray<FSchemaCategoryId> CategoryIds;
 };
 
 /**
@@ -371,13 +386,30 @@ struct FSchemaDescriptor
 struct FSchemaRegistryDescriptorConfig
 {
 	/**
-	  * List of all supported schema.
-	  */
+	 * List of all supported schema.
+	 */
 	TArray<FSchemaDescriptor> SchemaDescriptors;
 	/**
-	  * List of all service descriptors.
-	  */
+	 * List of all schema categories
+	 */
+	TArray<FSchemaCategoryDescriptor> SchemaCategoryDescriptors;
+	/**
+	 * List of all schema attributes
+	 */
+	TArray<FSchemaAttributeDescriptor> SchemaAttributeDescriptors;
+	/**
+	 * List of additional attributes to add to the given category for the given schema.
+	 */
+	TArray<FSchemaCategoryAttributesDescriptor> SchemaCategoryAttributeDescriptors;
+
+	/**
+	 * List of all service descriptors.
+	 */
 	TArray<FSchemaServiceDescriptor> ServiceDescriptors;
+	/**
+	 * List of all service attributes
+	 */
+	TArray<FSchemaServiceAttributeDescriptor> ServiceAttributeDescriptors;
 };
 
 /**
@@ -635,7 +667,7 @@ END_ONLINE_STRUCT_META()
 
 BEGIN_ONLINE_STRUCT_META(FSchemaServiceDescriptor)
 	ONLINE_STRUCT_FIELD(FSchemaServiceDescriptor, Id),
-	ONLINE_STRUCT_FIELD(FSchemaServiceDescriptor, Attributes)
+	ONLINE_STRUCT_FIELD(FSchemaServiceDescriptor, AttributeIds)
 END_ONLINE_STRUCT_META()
 
 BEGIN_ONLINE_STRUCT_META(FSchemaAttributeDescriptor)
@@ -648,19 +680,28 @@ END_ONLINE_STRUCT_META()
 
 BEGIN_ONLINE_STRUCT_META(FSchemaCategoryDescriptor)
 	ONLINE_STRUCT_FIELD(FSchemaCategoryDescriptor, Id),
-	ONLINE_STRUCT_FIELD(FSchemaCategoryDescriptor, ServiceDescriptorId),
-	ONLINE_STRUCT_FIELD(FSchemaCategoryDescriptor, Attributes)
+	ONLINE_STRUCT_FIELD(FSchemaCategoryDescriptor, ServiceDescriptorId)
 END_ONLINE_STRUCT_META()
 
 BEGIN_ONLINE_STRUCT_META(FSchemaDescriptor)
 	ONLINE_STRUCT_FIELD(FSchemaDescriptor, Id),
 	ONLINE_STRUCT_FIELD(FSchemaDescriptor, ParentId),
-	ONLINE_STRUCT_FIELD(FSchemaDescriptor, Categories)
+	ONLINE_STRUCT_FIELD(FSchemaDescriptor, CategoryIds)
+END_ONLINE_STRUCT_META()
+
+BEGIN_ONLINE_STRUCT_META(FSchemaCategoryAttributesDescriptor)
+	ONLINE_STRUCT_FIELD(FSchemaCategoryAttributesDescriptor, SchemaId),
+	ONLINE_STRUCT_FIELD(FSchemaCategoryAttributesDescriptor, CategoryId),
+	ONLINE_STRUCT_FIELD(FSchemaCategoryAttributesDescriptor, AttributeIds)
 END_ONLINE_STRUCT_META()
 
 BEGIN_ONLINE_STRUCT_META(FSchemaRegistryDescriptorConfig)
 	ONLINE_STRUCT_FIELD(FSchemaRegistryDescriptorConfig, SchemaDescriptors),
-	ONLINE_STRUCT_FIELD(FSchemaRegistryDescriptorConfig, ServiceDescriptors)
+	ONLINE_STRUCT_FIELD(FSchemaRegistryDescriptorConfig, SchemaCategoryDescriptors),
+	ONLINE_STRUCT_FIELD(FSchemaRegistryDescriptorConfig, SchemaAttributeDescriptors),
+	ONLINE_STRUCT_FIELD(FSchemaRegistryDescriptorConfig, SchemaCategoryAttributeDescriptors),
+	ONLINE_STRUCT_FIELD(FSchemaRegistryDescriptorConfig, ServiceDescriptors),
+	ONLINE_STRUCT_FIELD(FSchemaRegistryDescriptorConfig, ServiceAttributeDescriptors)
 END_ONLINE_STRUCT_META()
 
 /* Meta*/ }
