- Rename module implementation class from FOnlineServicesCommon to FOnlineServicesCommonModule to avoid conflicting with the other class with the same name.

diff --git a/Engine/Plugins/Online/Experimental/OnlineServices/Source/OnlineServicesCommon/Private/OnlineServicesCommonModule.cpp b/Engine/Plugins/Online/Experimental/OnlineServices/Source/OnlineServicesCommon/Private/OnlineServicesCommonModule.cpp
--- a/Engine/Plugins/Online/Experimental/OnlineServices/Source/OnlineServicesCommon/Private/OnlineServicesCommonModule.cpp
+++ b/Engine/Plugins/Online/Experimental/OnlineServices/Source/OnlineServicesCommon/Private/OnlineServicesCommonModule.cpp
@@ -3,7 +3,7 @@
 #include "CoreMinimal.h"
 #include "Modules/ModuleManager.h"
 
-class FOnlineServicesCommon : public IModuleInterface
+class FOnlineServicesCommonModule : public IModuleInterface
 {
 public:
 	virtual void StartupModule()
@@ -17,4 +17,4 @@ public:
 	}
 };
 
-IMPLEMENT_MODULE(FOnlineServicesCommon, OnlineServicesCommon);
+IMPLEMENT_MODULE(FOnlineServicesCommonModule, OnlineServicesCommon);
