- Structify StatsCommon config

diff --git a/Engine/Plugins/Online/Experimental/OnlineServices/Source/OnlineServicesCommon/Private/Online/StatsCommon.cpp b/Engine/Plugins/Online/Experimental/OnlineServices/Source/OnlineServicesCommon/Private/Online/StatsCommon.cpp
--- a/Engine/Plugins/Online/Experimental/OnlineServices/Source/OnlineServicesCommon/Private/Online/StatsCommon.cpp
+++ b/Engine/Plugins/Online/Experimental/OnlineServices/Source/OnlineServicesCommon/Private/Online/StatsCommon.cpp
@@ -48,38 +48,47 @@ void LexFromString(EStatModifyMethod& OutValue, const TCHAR* InStr)
 	}
 }
 
-const TCHAR* LexToString(EStatUsageFlags Value)
+const FString LexToString(EStatUsageFlags Value)
 {
-	switch (Value)
+	if (Value == EStatUsageFlags::None)
 	{
-	case EStatUsageFlags::Achievement:
-		return TEXT("Achievement");
-	case EStatUsageFlags::Leaderboard:
-		return TEXT("Leaderboard");
-	default: checkNoEntry(); // Intentional fallthrough
-	case EStatUsageFlags::None:
 		return TEXT("None");
 	}
-}
 
-void LexFromString(EStatUsageFlags& OutValue, const TCHAR* InStr)
-{
-	if (FCString::Stricmp(InStr, TEXT("None")) == 0)
+	TArray<FString> SetFlags;
+	if (EnumHasAllFlags(Value, EStatUsageFlags::Achievement))
 	{
-		OutValue = EStatUsageFlags::None;
+		SetFlags.Emplace(TEXT("Achievement"));
 	}
-	else if (FCString::Stricmp(InStr, TEXT("Achievement")) == 0)
+	if (EnumHasAllFlags(Value, EStatUsageFlags::Leaderboard))
 	{
-		OutValue = EStatUsageFlags::Achievement;
+		SetFlags.Emplace(TEXT("Leaderboard"));
 	}
-	else if (FCString::Stricmp(InStr, TEXT("Leaderboard")) == 0)
-	{
-		OutValue = EStatUsageFlags::Leaderboard;
-	}
-	else
+	return FString::Join(SetFlags, TEXT("+"));
+}
+
+void LexFromString(EStatUsageFlags& OutValue, const FString& InStr)
+{
+	OutValue = EStatUsageFlags::None;
+	if (InStr != TEXT("None"))
 	{
-		ensureMsgf(false, TEXT("Can't convert %s to EStatusageFlags"), InStr);
-		OutValue = EStatUsageFlags::None;
+		TArray<FString> SplitStrings;
+		InStr.ParseIntoArray(SplitStrings, TEXT("+"));
+		for (const FString& SplitString : SplitStrings)
+		{
+			if (SplitString == TEXT("Achievement"))
+			{
+				OutValue |= EStatUsageFlags::Achievement;
+			}
+			else if (SplitString == TEXT("Leaderboard"))
+			{
+				OutValue |= EStatUsageFlags::Leaderboard;
+			}
+			else
+			{
+				ensureMsgf(false, TEXT("Can't convert %s to EStatUsageFlags"), *SplitString);
+			}
+		}
 	}
 }
 
@@ -90,40 +99,13 @@ FStatsCommon::FStatsCommon(FOnlineServicesCommon& InServices)
 
 void FStatsCommon::UpdateConfig()
 {
-	const TCHAR* ConfigSection = TEXT("OnlineServices.Stats");
+	FStatsCommonConfig Config;
+	TOnlineComponent::LoadConfig(Config);
 
-	for (int StatIdx = 0;; StatIdx++)
+	for (FStatDefinition& StatDefinition : Config.StatDefinitions)
 	{
-		FString StatName;
-		GConfig->GetString(ConfigSection, *FString::Printf(TEXT("StatDef_%d_Name"), StatIdx), StatName, GEngineIni);
-		if (StatName.IsEmpty())
-		{
-			break;
-		}
-
-		FStatDefinition& StatDefinition = StatDefinitions.Emplace(StatName);
-		StatDefinition.Name = MoveTemp(StatName);
-
-		GConfig->GetInt(ConfigSection, *FString::Printf(TEXT("StatDef_%d_Id"), StatIdx), StatDefinition.Id, GEngineIni);
-
-		FText StatUsageFlagsStr;
-		GConfig->GetText(ConfigSection, *FString::Printf(TEXT("StatDef_%d_UsageFlags"), StatIdx), StatUsageFlagsStr, GEngineIni);
-		TArray<FString> StatUsageFlagStrArray;
-		bool CullEmpty = true;
-		StatUsageFlagsStr.ToString().ParseIntoArray(StatUsageFlagStrArray, TEXT(","), CullEmpty);
-		for (const FString& StatUsageFlagStr : StatUsageFlagStrArray)
-		{
-			EStatUsageFlags StatUsageFlag = EStatUsageFlags::None;
-			LexFromString(StatUsageFlag, *StatUsageFlagStr);
-			StatDefinition.UsageFlags |= (uint32)StatUsageFlag;
-		}
-
-		FText StatModifyMethod;
-		GConfig->GetText(ConfigSection, *FString::Printf(TEXT("StatDef_%d_ModifyMethod"), StatIdx), StatModifyMethod, GEngineIni);
-		if (!StatModifyMethod.IsEmpty())
-		{
-			LexFromString(StatDefinition.ModifyMethod, *StatModifyMethod.ToString());
-		}
+		FString StatName = StatDefinition.Name;
+		StatDefinitions.Emplace(MoveTemp(StatName), MoveTemp(StatDefinition));
 	}
 }
 
diff --git a/Engine/Plugins/Online/Experimental/OnlineServices/Source/OnlineServicesCommon/Public/Online/StatsCommon.h b/Engine/Plugins/Online/Experimental/OnlineServices/Source/OnlineServicesCommon/Public/Online/StatsCommon.h
--- a/Engine/Plugins/Online/Experimental/OnlineServices/Source/OnlineServicesCommon/Public/Online/StatsCommon.h
+++ b/Engine/Plugins/Online/Experimental/OnlineServices/Source/OnlineServicesCommon/Public/Online/StatsCommon.h
@@ -31,8 +31,8 @@ enum class EStatUsageFlags : uint8
 
 ENUM_CLASS_FLAGS(EStatUsageFlags);
 
-const TCHAR* LexToString(EStatUsageFlags Value);
-void LexFromString(EStatUsageFlags& OutValue, const TCHAR* InStr);
+const FString LexToString(EStatUsageFlags Value);
+void LexFromString(EStatUsageFlags& OutValue, const FString& InStr);
 
 struct FStatDefinition
 {
@@ -40,14 +40,32 @@ struct FStatDefinition
 	FString Name;
 	/* Corresponding stat id on the platform if needed */
 	int32 Id = 0;
-
-	/* What is this stat used for, in array format split by ",", for example: "Achievement,Leaderboard" */
-	int32 UsageFlags = 0;
-
+	/* What is this stat used for */
+	EStatUsageFlags UsageFlags = EStatUsageFlags::None;
 	/* How the stat will be modified, only useful when EStatUsageFlags::Achievement is set in UsageFlags */
 	EStatModifyMethod ModifyMethod = EStatModifyMethod::Set;
 };
 
+struct FStatsCommonConfig
+{
+	TArray<FStatDefinition> StatDefinitions;
+};
+
+namespace Meta {
+
+BEGIN_ONLINE_STRUCT_META(FStatDefinition)
+	ONLINE_STRUCT_FIELD(FStatDefinition, Name),
+	ONLINE_STRUCT_FIELD(FStatDefinition, Id),
+	ONLINE_STRUCT_FIELD(FStatDefinition, UsageFlags),
+	ONLINE_STRUCT_FIELD(FStatDefinition, ModifyMethod)
+END_ONLINE_STRUCT_META()
+
+BEGIN_ONLINE_STRUCT_META(FStatsCommonConfig)
+	ONLINE_STRUCT_FIELD(FStatsCommonConfig, StatDefinitions)
+END_ONLINE_STRUCT_META()
+
+/* Meta */ }
+
 struct FFindUserStatsByAccountId
 {
 	FFindUserStatsByAccountId(const FAccountId& InAccountId)
diff --git a/Engine/Plugins/Online/Experimental/OnlineServicesNull/Source/Private/Online/StatsNull.cpp b/Engine/Plugins/Online/Experimental/OnlineServicesNull/Source/Private/Online/StatsNull.cpp
--- a/Engine/Plugins/Online/Experimental/OnlineServicesNull/Source/Private/Online/StatsNull.cpp
+++ b/Engine/Plugins/Online/Experimental/OnlineServicesNull/Source/Private/Online/StatsNull.cpp
@@ -36,7 +36,7 @@ TOnlineAsyncOpHandle<FUpdateStats> FStatsNull::UpdateStats(FUpdateStats::Params&
 			TMap<FString, FStatValue>& UserStats = ExistingUserStats->Stats;
 			for (const auto& UpdateUserStatPair : UpdateUserStats.Stats)
 			{
-				if (FStatDefinition* StatDefinition = StatDefinitions.Find(UpdateUserStatPair.Key))
+				if (const FStatDefinition* StatDefinition = GetStatDefinition(UpdateUserStatPair.Key))
 				{
 					if (FStatValue* StatValue = UserStats.Find(UpdateUserStatPair.Key))
 					{
@@ -82,7 +82,7 @@ TOnlineAsyncOpHandle<FUpdateStats> FStatsNull::UpdateStats(FUpdateStats::Params&
 						UserStats.Emplace(UpdateUserStatPair.Key, UpdateUserStatPair.Value);
 					}
 
-					if (StatDefinition->UsageFlags & (uint32)EStatUsageFlags::Leaderboard)
+					if (EnumHasAnyFlags(StatDefinition->UsageFlags, EStatUsageFlags::Leaderboard))
 					{
 						// TODO: Call LeaderboardsNull to update
 					}
