- Structify LeaderboardsCommon config

diff --git a/Engine/Plugins/Online/Experimental/OnlineServices/Source/OnlineServicesCommon/Private/Online/LeaderboardsCommon.cpp b/Engine/Plugins/Online/Experimental/OnlineServices/Source/OnlineServicesCommon/Private/Online/LeaderboardsCommon.cpp
--- a/Engine/Plugins/Online/Experimental/OnlineServices/Source/OnlineServicesCommon/Private/Online/LeaderboardsCommon.cpp
+++ b/Engine/Plugins/Online/Experimental/OnlineServices/Source/OnlineServicesCommon/Private/Online/LeaderboardsCommon.cpp
@@ -82,36 +82,15 @@ void FLeaderboardsCommon::UpdateConfig()
 {
 	Super::UpdateConfig();
 
-	const TCHAR* ConfigSection = TEXT("OnlineServices.Leaderboards");
-	GConfig->GetBool(ConfigSection, TEXT("LeaderboardDef_IsTitleManaged"), bIsTitleManaged, GEngineIni);
+	FLeaderboardsCommonConfig Config;
+	TOnlineComponent::LoadConfig(Config);
 
-	for (int StatIdx = 0;; StatIdx++)
-	{
-		FString LeaderboardName;
-		GConfig->GetString(ConfigSection, *FString::Printf(TEXT("LeaderboardDef_%d_Name"), StatIdx), LeaderboardName, GEngineIni);
-		if (LeaderboardName.IsEmpty())
-		{
-			break;
-		}
+	bIsTitleManaged = Config.bIsTitleManaged;
 
-		FLeaderboardDefinition& LeaderboardDefinition = LeaderboardDefinitions.Emplace(LeaderboardName);
-		LeaderboardDefinition.Name = MoveTemp(LeaderboardName);
-
-		GConfig->GetInt(ConfigSection, *FString::Printf(TEXT("LeaderboardDef_%d_Id"), StatIdx), LeaderboardDefinition.Id, GEngineIni);
-
-		FString LeaderboardUpdateMethod;
-		GConfig->GetString(ConfigSection, *FString::Printf(TEXT("LeaderboardDef_%d_UpdateMethod"), StatIdx), LeaderboardUpdateMethod, GEngineIni);
-		if (!LeaderboardUpdateMethod.IsEmpty())
-		{
-			LexFromString(LeaderboardDefinition.LeaderboardUpdateMethod, *LeaderboardUpdateMethod);
-		}
-
-		FString LeaderboardOrderMethod;
-		GConfig->GetString(ConfigSection, *FString::Printf(TEXT("LeaderboardDef_%d_OrderMethod"), StatIdx), LeaderboardOrderMethod, GEngineIni);
-		if (!LeaderboardOrderMethod.IsEmpty())
-		{
-			LexFromString(LeaderboardDefinition.LeaderboardOrderMethod, *LeaderboardOrderMethod);
-		}
+	for (FLeaderboardDefinition& LeaderboardDefinition : Config.LeaderboardDefinitions)
+	{
+		FString LeaderboardName = LeaderboardDefinition.Name;
+		LeaderboardDefinitions.Emplace(MoveTemp(LeaderboardName), MoveTemp(LeaderboardDefinition));
 	}
 }
 
diff --git a/Engine/Plugins/Online/Experimental/OnlineServices/Source/OnlineServicesCommon/Public/Online/LeaderboardsCommon.h b/Engine/Plugins/Online/Experimental/OnlineServices/Source/OnlineServicesCommon/Public/Online/LeaderboardsCommon.h
--- a/Engine/Plugins/Online/Experimental/OnlineServices/Source/OnlineServicesCommon/Public/Online/LeaderboardsCommon.h
+++ b/Engine/Plugins/Online/Experimental/OnlineServices/Source/OnlineServicesCommon/Public/Online/LeaderboardsCommon.h
@@ -34,13 +34,35 @@ struct FLeaderboardDefinition
 	FString Name;
 	/* Corresponding leaderboard id on the platform if needed */
 	int32 Id = 0;
-
 	/* How the leaderboard score will be updated */
-	ELeaderboardUpdateMethod LeaderboardUpdateMethod = ELeaderboardUpdateMethod::Force;
+	ELeaderboardUpdateMethod UpdateMethod = ELeaderboardUpdateMethod::Force;
 	/* How the leaderboard score will be ordered */
-	ELeaderboardOrderMethod LeaderboardOrderMethod = ELeaderboardOrderMethod::Descending;
+	ELeaderboardOrderMethod OrderMethod = ELeaderboardOrderMethod::Descending;
 };
 
+struct FLeaderboardsCommonConfig
+{
+	bool bIsTitleManaged = false;
+	TArray<FLeaderboardDefinition> LeaderboardDefinitions;
+};
+
+namespace Meta {
+
+BEGIN_ONLINE_STRUCT_META(FLeaderboardDefinition)
+	ONLINE_STRUCT_FIELD(FLeaderboardDefinition, Name),
+	ONLINE_STRUCT_FIELD(FLeaderboardDefinition, Id),
+	ONLINE_STRUCT_FIELD(FLeaderboardDefinition, UpdateMethod),
+	ONLINE_STRUCT_FIELD(FLeaderboardDefinition, OrderMethod)
+END_ONLINE_STRUCT_META()
+
+BEGIN_ONLINE_STRUCT_META(FLeaderboardsCommonConfig)
+	ONLINE_STRUCT_FIELD(FLeaderboardsCommonConfig, bIsTitleManaged),
+	ONLINE_STRUCT_FIELD(FLeaderboardsCommonConfig, LeaderboardDefinitions)
+END_ONLINE_STRUCT_META()
+
+
+/* Meta */ }
+
 struct FWriteLeaderboardScores
 {
 	static constexpr TCHAR Name[] = TEXT("WriteLeaderboardScores");
diff --git a/Engine/Plugins/Online/Experimental/OnlineServicesNull/Source/Private/Online/LeaderboardsNull.cpp b/Engine/Plugins/Online/Experimental/OnlineServicesNull/Source/Private/Online/LeaderboardsNull.cpp
--- a/Engine/Plugins/Online/Experimental/OnlineServicesNull/Source/Private/Online/LeaderboardsNull.cpp
+++ b/Engine/Plugins/Online/Experimental/OnlineServicesNull/Source/Private/Online/LeaderboardsNull.cpp
@@ -32,10 +32,10 @@ namespace Private
 
 uint64 GetUpdatedScoreByLeaderboardDefinition(FLeaderboardDefinition* LeaderboardDefinition, uint64 OldScore, uint64 NewScore)
 {
-	switch (LeaderboardDefinition->LeaderboardUpdateMethod)
+	switch (LeaderboardDefinition->UpdateMethod)
 	{
 	case ELeaderboardUpdateMethod::KeepBest:
-		switch (LeaderboardDefinition->LeaderboardOrderMethod)
+		switch (LeaderboardDefinition->OrderMethod)
 		{
 		case ELeaderboardOrderMethod::Descending:
 			return FMath::Max(NewScore, OldScore);
@@ -110,7 +110,7 @@ TOnlineAsyncOpHandle<FWriteLeaderboardScores> FLeaderboardsNull::WriteLeaderboar
 			const FUserScoreNull& UserScore = CurrentNode->GetValue();
 
 			bool ShouldInsert = false;
-			switch (LeaderboardDefinition->LeaderboardOrderMethod)
+			switch (LeaderboardDefinition->OrderMethod)
 			{
 			case ELeaderboardOrderMethod::Descending:
 				ShouldInsert = UpdatedScore > UserScore.Score;
