Fix for bug in Session serialization

diff --git a/Engine/Plugins/Online/Experimental/OnlineServices/Source/OnlineServicesCommon/Private/Online/NboSerializerCommonSvc.cpp b/Engine/Plugins/Online/Experimental/OnlineServices/Source/OnlineServicesCommon/Private/Online/NboSerializerCommonSvc.cpp
--- a/Engine/Plugins/Online/Experimental/OnlineServices/Source/OnlineServicesCommon/Private/Online/NboSerializerCommonSvc.cpp
+++ b/Engine/Plugins/Online/Experimental/OnlineServices/Source/OnlineServicesCommon/Private/Online/NboSerializerCommonSvc.cpp
@@ -84,7 +84,7 @@ void SerializeToBuffer(FNboSerializeToBuffer& Packet, const FSessionSettings& Se
 
 void SerializeToBuffer(FNboSerializeToBuffer& Packet, const FSessionCommon& Session)
 {
-	// Session.OwnerUserId will be serialized in implementations as user id types will vary
+	// Session.OwnerAccountId will be serialized in implementations as user id types will vary
 	SerializeToBuffer(Packet, Session.SessionInfo);
 	SerializeToBuffer(Packet, Session.SessionSettings);
 	// Session.SessionMembers will be serialized in implementations as user id types will vary
@@ -164,18 +164,10 @@ void SerializeFromBuffer(FNboSerializeFromBuffer& Packet, FSessionMember& Sessio
 
 void SerializeFromBuffer(FNboSerializeFromBuffer& Packet, FSessionInfo& SessionInfo)
 {
-	uint8 Read = 0;
-
-	// Read all the booleans as bytes
-	Packet >> Read;
-	SessionInfo.bAllowSanctionedPlayers = !!Read;
-	Packet >> Read;
-	SessionInfo.bAntiCheatProtected = !!Read;
-	Packet >> Read;
-	SessionInfo.bIsDedicatedServerSession = !!Read;
-	Packet >> Read;
-	SessionInfo.bIsLANSession = !!Read;
-
+	Packet >> SessionInfo.bAllowSanctionedPlayers;
+	Packet >> SessionInfo.bAntiCheatProtected;
+	Packet >> SessionInfo.bIsDedicatedServerSession;
+	Packet >> SessionInfo.bIsLANSession;
 	Packet >> SessionInfo.SessionIdOverride;
 
 	// SessionInfo.SessionId will be deserialized in implementations as user id types will vary
@@ -183,14 +175,11 @@ void SerializeFromBuffer(FNboSerializeFromBuffer& Packet, FSessionInfo& SessionI
 
 void SerializeFromBuffer(FNboSerializeFromBuffer& Packet, FSessionSettings& SessionSettings)
 {
-	uint8 Read = 0;
-
-	// Read all the booleans as bytes
-	Packet >> Read;
-	SessionSettings.bAllowNewMembers = !!Read;
+	Packet >> SessionSettings.bAllowNewMembers;
 
 	SerializeFromBuffer(Packet, SessionSettings.CustomSettings);
 
+	uint8 Read = 0;
 	Packet >> Read;
 	SessionSettings.JoinPolicy = (ESessionJoinPolicy)Read;
 
diff --git a/Engine/Plugins/Online/OnlineBase/Source/Public/Online/NboSerializer.h b/Engine/Plugins/Online/OnlineBase/Source/Public/Online/NboSerializer.h
--- a/Engine/Plugins/Online/OnlineBase/Source/Public/Online/NboSerializer.h
+++ b/Engine/Plugins/Online/OnlineBase/Source/Public/Online/NboSerializer.h
@@ -91,6 +91,16 @@ public:
 		bHasOverflowed = false;
 	}
 
+	/**
+	 * Adds a bool to the buffer (converted to uint8)
+	 */
+	friend inline FNboSerializeToBuffer& operator<<(FNboSerializeToBuffer& Ar, const bool B)
+	{
+		Ar << (uint8)B;
+
+		return Ar;
+	}
+
 	/**
 	 * Adds a char to the buffer
 	 */
@@ -417,6 +427,18 @@ public:
 	{
 	}
 
+	/**
+	 * Reads a bool from the buffer (as uint8)
+	 */
+	friend inline FNboSerializeFromBuffer& operator>>(FNboSerializeFromBuffer& Ar, bool& B)
+	{
+		uint8 Read = 0;
+		Ar >> Read;
+		B = !!Read;
+
+		return Ar;
+	}
+
 	/**
 	 * Reads a char from the buffer
 	 */
