- Structify AchievementsCommon config

diff --git a/Engine/Plugins/Online/Experimental/OnlineServices/Source/OnlineServicesCommon/Private/Online/AchievementsCommon.cpp b/Engine/Plugins/Online/Experimental/OnlineServices/Source/OnlineServicesCommon/Private/Online/AchievementsCommon.cpp
--- a/Engine/Plugins/Online/Experimental/OnlineServices/Source/OnlineServicesCommon/Private/Online/AchievementsCommon.cpp
+++ b/Engine/Plugins/Online/Experimental/OnlineServices/Source/OnlineServicesCommon/Private/Online/AchievementsCommon.cpp
@@ -11,7 +11,7 @@ namespace UE::Online {
 
 bool FAchievementUnlockRule::ContainsStat(const FString& StatName) const
 {
-	return Conditions.ContainsByPredicate([&StatName](const FAchievementUnlockCondition& Condition) { return Condition.StatToCheck == StatName; });
+	return Conditions.ContainsByPredicate([&StatName](const FAchievementUnlockCondition& Condition) { return Condition.StatName == StatName; });
 }
 
 FAchievementsCommon::FAchievementsCommon(FOnlineServicesCommon& InServices)
@@ -36,44 +36,7 @@ void FAchievementsCommon::Shutdown()
 void FAchievementsCommon::UpdateConfig()
 {
 	TOnlineComponent<IAchievements>::UpdateConfig();
-
-	const TCHAR* ConfigSection = TEXT("OnlineServices.AchievementUnlockRules");
-	GConfig->GetBool(ConfigSection, TEXT("AchievementDef_IsTitleManaged"), bIsTitleManaged, GEngineIni);
-
-	if (bIsTitleManaged)
-	{
-		for (int RuleIdx = 0;; RuleIdx++)
-		{
-			FString AchievementToUnlock;
-			GConfig->GetString(ConfigSection, *FString::Printf(TEXT("AchievementDef_%d_UnlockRule_AchievementToUnlock"), RuleIdx), AchievementToUnlock, GEngineIni);
-
-			if (AchievementToUnlock.IsEmpty())
-			{
-				break;
-			}
-
-			FAchievementUnlockRule& AchievementUnlockRule = AchievementUnlockRules.Emplace_GetRef();
-			AchievementUnlockRule.AchievementToUnlock = MoveTemp(AchievementToUnlock);
-
-			FString ConditionsStr;
-			GConfig->GetString(ConfigSection, *FString::Printf(TEXT("AchievementDef_%d_UnlockRule_Conditions"), RuleIdx), ConditionsStr, GEngineIni);
-
-			bool bCullEmpty = true;
-			TArray<FString> ConditionStrs;
-			ConditionsStr.ParseIntoArray(ConditionStrs, TEXT(","), bCullEmpty);
-
-			for (const FString& ConditionStr : ConditionStrs)
-			{
-				FString StatToCheck;
-				FString UnlockThreshold;
-				ConditionStr.Split(TEXT(":"), &StatToCheck, &UnlockThreshold);
-
-				FAchievementUnlockCondition& Condition = AchievementUnlockRule.Conditions.Emplace_GetRef();
-				Condition.StatToCheck = MoveTemp(StatToCheck);
-				Condition.UnlockThreshold = (int64)FCString::Atoi(*UnlockThreshold);
-			}
-		}
-	}
+	TOnlineComponent<IAchievements>::LoadConfig(Config);
 }
 
 void FAchievementsCommon::RegisterCommands()
@@ -135,7 +98,7 @@ TOnlineEvent<void(const FAchievementStateUpdated&)> FAchievementsCommon::OnAchie
 
 void FAchievementsCommon::UnlockAchievementsByStats(const FStatsUpdated& StatsUpdated)
 {
-	if (!bIsTitleManaged)
+	if (!Config.bIsTitleManaged)
 	{
 		return;
 	}
@@ -146,13 +109,13 @@ void FAchievementsCommon::UnlockAchievementsByStats(const FStatsUpdated& StatsUp
 	{
 		for (const TPair<FString, FStatValue>& StatPair : UserStats.Stats)
 		{
-			for (const FAchievementUnlockRule& AchievementUnlockRule : AchievementUnlockRules)
+			for (const FAchievementUnlockRule& AchievementUnlockRule : Config.UnlockRules)
 			{
 				if (AchievementUnlockRule.ContainsStat(StatPair.Key))
 				{
 					for (const FAchievementUnlockCondition& Condition : AchievementUnlockRule.Conditions)
 					{
-						StatNames.AddUnique(Condition.StatToCheck);
+						StatNames.AddUnique(Condition.StatName);
 					}
 				}
 			}
@@ -199,27 +162,37 @@ void FAchievementsCommon::UnlockAchievementsByStats(const FStatsUpdated& StatsUp
 
 void FAchievementsCommon::ExecuteUnlockRulesRelatedToStat(const FAccountId& AccountId, const FString& StatName, const TMap<FString, FStatValue>& Stats, TArray<FString>& OutAchievementsToUnlock)
 {
-	for (const FAchievementUnlockRule& AchievementUnlockRule : AchievementUnlockRules)
+	if (!Config.bIsTitleManaged)
+	{
+		return;
+	}
+
+	for (const FAchievementUnlockRule& AchievementUnlockRule : Config.UnlockRules)
 	{
 		if (AchievementUnlockRule.ContainsStat(StatName)
-			&& !IsUnlocked(AccountId, AchievementUnlockRule.AchievementToUnlock) 
+			&& !IsUnlocked(AccountId, AchievementUnlockRule.AchievementId) 
 			&& MeetUnlockCondition(AchievementUnlockRule, Stats))
 		{
-			OutAchievementsToUnlock.AddUnique(AchievementUnlockRule.AchievementToUnlock);
+			OutAchievementsToUnlock.AddUnique(AchievementUnlockRule.AchievementId);
 		}
 	}
 }
 
-bool FAchievementsCommon::MeetUnlockCondition(FAchievementUnlockRule AchievementUnlockRule, const TMap<FString, FStatValue>& Stats)
+bool FAchievementsCommon::MeetUnlockCondition(const FAchievementUnlockRule& AchievementUnlockRule, const TMap<FString, FStatValue>& Stats)
 {
+	if (!Config.bIsTitleManaged)
+	{
+		return false;
+	}
+
 	for (const FAchievementUnlockCondition& Condition : AchievementUnlockRule.Conditions)
 	{
-		if (const FStatDefinition* StatDefinition = Services.Get<FStatsCommon>()->GetStatDefinition(Condition.StatToCheck))
+		if (const FStatDefinition* StatDefinition = Services.Get<FStatsCommon>()->GetStatDefinition(Condition.StatName))
 		{
-			const FStatValue* StatValue = Stats.Find(Condition.StatToCheck);
+			const FStatValue* StatValue = Stats.Find(Condition.StatName);
 			if (!StatValue)
 			{
-				UE_LOG(LogTemp, Warning, TEXT("Can't find stat %s when check if it can unlock achievement."), *Condition.StatToCheck);
+				UE_LOG(LogTemp, Warning, TEXT("Can't find stat %s when check if it can unlock achievement."), *Condition.StatName);
 				return false;
 			}
 
diff --git a/Engine/Plugins/Online/Experimental/OnlineServices/Source/OnlineServicesCommon/Public/Online/AchievementsCommon.h b/Engine/Plugins/Online/Experimental/OnlineServices/Source/OnlineServicesCommon/Public/Online/AchievementsCommon.h
--- a/Engine/Plugins/Online/Experimental/OnlineServices/Source/OnlineServicesCommon/Public/Online/AchievementsCommon.h
+++ b/Engine/Plugins/Online/Experimental/OnlineServices/Source/OnlineServicesCommon/Public/Online/AchievementsCommon.h
@@ -12,18 +12,44 @@ class FOnlineServicesCommon;
 
 struct FAchievementUnlockCondition
 {
-	FString StatToCheck;
+	FString StatName;
 	FStatValue UnlockThreshold; // The unlock rule depends on Stat modification type
 };
 
 struct FAchievementUnlockRule
 {
-	FString AchievementToUnlock;
+	FString AchievementId;
 	TArray<FAchievementUnlockCondition> Conditions;
 
 	bool ContainsStat(const FString& StatName) const;
 };
 
+struct FAchievementsCommonConfig
+{
+	bool bIsTitleManaged = false;
+	TArray<FAchievementUnlockRule> UnlockRules;
+};
+
+namespace Meta
+{
+
+BEGIN_ONLINE_STRUCT_META(FAchievementUnlockCondition)
+	ONLINE_STRUCT_FIELD(FAchievementUnlockCondition, StatName),
+	ONLINE_STRUCT_FIELD(FAchievementUnlockCondition, UnlockThreshold)
+END_ONLINE_STRUCT_META()
+
+BEGIN_ONLINE_STRUCT_META(FAchievementUnlockRule)
+	ONLINE_STRUCT_FIELD(FAchievementUnlockRule, AchievementId),
+	ONLINE_STRUCT_FIELD(FAchievementUnlockRule, Conditions)
+END_ONLINE_STRUCT_META()
+
+BEGIN_ONLINE_STRUCT_META(FAchievementsCommonConfig)
+	ONLINE_STRUCT_FIELD(FAchievementsCommonConfig, bIsTitleManaged),
+	ONLINE_STRUCT_FIELD(FAchievementsCommonConfig, UnlockRules)
+END_ONLINE_STRUCT_META()
+
+/* Meta */ }
+
 class ONLINESERVICESCOMMON_API FAchievementsCommon : public TOnlineComponent<IAchievements>
 {
 public:
@@ -52,13 +78,12 @@ protected:
 
 	void UnlockAchievementsByStats(const FStatsUpdated& StatsUpdated);
 	void ExecuteUnlockRulesRelatedToStat(const FAccountId& AccountId, const FString& StatName, const TMap<FString, FStatValue>& Stats, TArray<FString>& OutAchievementsToUnlock);
-	bool MeetUnlockCondition(FAchievementUnlockRule AchievementUnlockRule, const TMap<FString, FStatValue>& Stats);
+	bool MeetUnlockCondition(const FAchievementUnlockRule& AchievementUnlockRule, const TMap<FString, FStatValue>& Stats);
 	bool IsUnlocked(const FAccountId& AccountId, const FString& AchievementName) const;
 
-	TArray<FAchievementUnlockRule> AchievementUnlockRules;
 	FOnlineEventDelegateHandle StatEventHandle;
 
-	bool bIsTitleManaged = false;
+	FAchievementsCommonConfig Config;
 };
 
 /* UE::Online */ }
