OSSv2: Remove the UsageFlags of Stats, since it's not necessary for now.

diff --git a/Engine/Plugins/Online/Experimental/OnlineServices/Source/OnlineServicesCommon/Private/Online/StatsCommon.cpp b/Engine/Plugins/Online/Experimental/OnlineServices/Source/OnlineServicesCommon/Private/Online/StatsCommon.cpp
--- a/Engine/Plugins/Online/Experimental/OnlineServices/Source/OnlineServicesCommon/Private/Online/StatsCommon.cpp
+++ b/Engine/Plugins/Online/Experimental/OnlineServices/Source/OnlineServicesCommon/Private/Online/StatsCommon.cpp
@@ -48,50 +48,6 @@ void LexFromString(EStatModifyMethod& OutValue, const TCHAR* InStr)
 	}
 }
 
-const FString LexToString(EStatUsageFlags Value)
-{
-	if (Value == EStatUsageFlags::None)
-	{
-		return TEXT("None");
-	}
-
-	TArray<FString> SetFlags;
-	if (EnumHasAllFlags(Value, EStatUsageFlags::Achievement))
-	{
-		SetFlags.Emplace(TEXT("Achievement"));
-	}
-	if (EnumHasAllFlags(Value, EStatUsageFlags::Leaderboard))
-	{
-		SetFlags.Emplace(TEXT("Leaderboard"));
-	}
-	return FString::Join(SetFlags, TEXT("+"));
-}
-
-void LexFromString(EStatUsageFlags& OutValue, const FString& InStr)
-{
-	OutValue = EStatUsageFlags::None;
-	if (InStr != TEXT("None"))
-	{
-		TArray<FString> SplitStrings;
-		InStr.ParseIntoArray(SplitStrings, TEXT("+"));
-		for (const FString& SplitString : SplitStrings)
-		{
-			if (SplitString == TEXT("Achievement"))
-			{
-				OutValue |= EStatUsageFlags::Achievement;
-			}
-			else if (SplitString == TEXT("Leaderboard"))
-			{
-				OutValue |= EStatUsageFlags::Leaderboard;
-			}
-			else
-			{
-				ensureMsgf(false, TEXT("Can't convert %s to EStatUsageFlags"), *SplitString);
-			}
-		}
-	}
-}
-
 FStatsCommon::FStatsCommon(FOnlineServicesCommon& InServices)
 	: TOnlineComponent(TEXT("Stats"), InServices)
 {
diff --git a/Engine/Plugins/Online/Experimental/OnlineServices/Source/OnlineServicesCommon/Public/Online/StatsCommon.h b/Engine/Plugins/Online/Experimental/OnlineServices/Source/OnlineServicesCommon/Public/Online/StatsCommon.h
--- a/Engine/Plugins/Online/Experimental/OnlineServices/Source/OnlineServicesCommon/Public/Online/StatsCommon.h
+++ b/Engine/Plugins/Online/Experimental/OnlineServices/Source/OnlineServicesCommon/Public/Online/StatsCommon.h
@@ -22,27 +22,13 @@ enum class EStatModifyMethod : uint8
 const TCHAR* LexToString(EStatModifyMethod Value);
 void LexFromString(EStatModifyMethod& OutValue, const TCHAR* InStr);
 
-enum class EStatUsageFlags : uint8
-{
-	None = 0,
-	Achievement = 1 << 0,
-	Leaderboard = 1 << 1
-};
-
-ENUM_CLASS_FLAGS(EStatUsageFlags);
-
-const FString LexToString(EStatUsageFlags Value);
-void LexFromString(EStatUsageFlags& OutValue, const FString& InStr);
-
 struct FStatDefinition
 {
 	/* The name of the stat */
 	FString Name;
 	/* Corresponding stat id on the platform if needed */
 	int32 Id = 0;
-	/* What is this stat used for */
-	EStatUsageFlags UsageFlags = EStatUsageFlags::None;
-	/* How the stat will be modified, only useful when EStatUsageFlags::Achievement is set in UsageFlags */
+	/* How the stat will be modified */
 	EStatModifyMethod ModifyMethod = EStatModifyMethod::Set;
 };
 
@@ -56,7 +42,6 @@ namespace Meta {
 BEGIN_ONLINE_STRUCT_META(FStatDefinition)
 	ONLINE_STRUCT_FIELD(FStatDefinition, Name),
 	ONLINE_STRUCT_FIELD(FStatDefinition, Id),
-	ONLINE_STRUCT_FIELD(FStatDefinition, UsageFlags),
 	ONLINE_STRUCT_FIELD(FStatDefinition, ModifyMethod)
 END_ONLINE_STRUCT_META()
 
diff --git a/Engine/Plugins/Online/Experimental/OnlineServices/Source/OnlineServicesInterface/Public/Online/Stats.h b/Engine/Plugins/Online/Experimental/OnlineServices/Source/OnlineServicesInterface/Public/Online/Stats.h
--- a/Engine/Plugins/Online/Experimental/OnlineServices/Source/OnlineServicesInterface/Public/Online/Stats.h
+++ b/Engine/Plugins/Online/Experimental/OnlineServices/Source/OnlineServicesInterface/Public/Online/Stats.h
@@ -153,8 +153,8 @@ public:
 	virtual TOnlineEvent<void(const FStatsUpdated&)> OnStatsUpdated() = 0;
 
 	/**
-	 * Retrieve cached users' stats, which was retrieved when call QueryStats or BatchQueryStats, or when 
-	 * call UpdateStats for StatsNull impl. When choose to use StatsNull, make sure to listen to OnStatsUpdated
+	 * Retrieve cached users' stats, which was retrieved when call QueryStats or BatchQueryStats, or after 
+	 * calling UpdateStats. When choose to use StatsNull, make sure to listen to OnStatsUpdated
 	 * event and use this function to get all latest local stats and put them into SaveGame
 	 */
 	virtual TOnlineResult<FGetCachedStats> GetCachedStats(FGetCachedStats::Params&& Params) const = 0;
diff --git a/Engine/Plugins/Online/Experimental/OnlineServicesNull/Source/Private/Online/StatsNull.cpp b/Engine/Plugins/Online/Experimental/OnlineServicesNull/Source/Private/Online/StatsNull.cpp
--- a/Engine/Plugins/Online/Experimental/OnlineServicesNull/Source/Private/Online/StatsNull.cpp
+++ b/Engine/Plugins/Online/Experimental/OnlineServicesNull/Source/Private/Online/StatsNull.cpp
@@ -81,11 +81,6 @@ TOnlineAsyncOpHandle<FUpdateStats> FStatsNull::UpdateStats(FUpdateStats::Params&
 					{
 						UserStats.Emplace(UpdateUserStatPair.Key, UpdateUserStatPair.Value);
 					}
-
-					if (EnumHasAnyFlags(StatDefinition->UsageFlags, EStatUsageFlags::Leaderboard))
-					{
-						// TODO: Call LeaderboardsNull to update
-					}
 				}
 			}
 
