Documentation pass: Made sure all structs and methods in Sessions.h have descriptive comments.
Implemented final versions of ToLogString and DumpState for ISession and subclasses.
Renaming of SessionMemberIdsSet for clarity.
Miscellaneous fixes.

diff --git a/Engine/Plugins/Online/Experimental/OnlineServices/Source/OnlineServicesCommon/Private/Online/SessionsCommon.cpp b/Engine/Plugins/Online/Experimental/OnlineServices/Source/OnlineServicesCommon/Private/Online/SessionsCommon.cpp
--- a/Engine/Plugins/Online/Experimental/OnlineServices/Source/OnlineServicesCommon/Private/Online/SessionsCommon.cpp
+++ b/Engine/Plugins/Online/Experimental/OnlineServices/Source/OnlineServicesCommon/Private/Online/SessionsCommon.cpp
@@ -7,6 +7,7 @@
 
 namespace UE::Online {
 	/** FSessionCommon */
+
 	FSessionCommon& FSessionCommon::operator+=(const FSessionUpdate& SessionUpdate)
 	{
 		if (SessionUpdate.OwnerAccountId.IsSet())
@@ -31,7 +32,61 @@ namespace UE::Online {
 		return *this;
 	}
 
+	FString FSessionCommon::ToLogString() const
+	{
+		return FString::Printf(TEXT("ISession: SessionId [%s]"), *UE::Online::ToLogString(SessionInfo.SessionId));
+	}
+
+	void FSessionCommon::DumpState() const
+	{
+		UE_LOG(LogTemp, Log, TEXT("[ISession]"));
+		DumpMemberData();
+		DumpSessionInfo();
+		DumpSessionSettings();
+	}
+
+	void FSessionCommon::DumpMemberData() const
+	{
+		UE_LOG(LogTemp, Log, TEXT("OwnerAccountId [%s]"), *UE::Online::ToLogString(OwnerAccountId));
+		UE_LOG(LogTemp, Log, TEXT("SessionMemberIdsSet:"));
+
+		for (const FAccountId& SessionMemberId : SessionMembers)
+		{
+			UE_LOG(LogTemp, Log, TEXT("	[%s]"), *UE::Online::ToLogString(SessionMemberId));
+		}
+	}
+
+	void FSessionCommon::DumpSessionInfo() const
+	{
+		UE_LOG(LogTemp, Log, TEXT("SessionInfo:"));
+		UE_LOG(LogTemp, Log, TEXT("	SessionId [%s]"), *UE::Online::ToLogString(SessionInfo.SessionId));
+		UE_LOG(LogTemp, Log, TEXT("	bAllowSanctionedPlayers [%d]"), SessionInfo.bAllowSanctionedPlayers);
+		UE_LOG(LogTemp, Log, TEXT("	bAntiCheatProtected [%d]"), SessionInfo.bAntiCheatProtected);
+		UE_LOG(LogTemp, Log, TEXT("	bIsDedicatedServerSession [%d]"), SessionInfo.bIsDedicatedServerSession);
+		UE_LOG(LogTemp, Log, TEXT("	bIsLANSession [%d]"), SessionInfo.bIsLANSession);
+		UE_LOG(LogTemp, Log, TEXT("	SessionIdOverride [%s]"), *SessionInfo.SessionIdOverride);
+	}
+
+	void FSessionCommon::DumpSessionSettings() const
+	{
+		UE_LOG(LogTemp, Log, TEXT("SessionSettings:"));
+		UE_LOG(LogTemp, Log, TEXT("	bAllowNewMembers [%d]"), SessionSettings.bAllowNewMembers);
+		UE_LOG(LogTemp, Log, TEXT("	JoinPolicy [%s]"), LexToString(SessionSettings.JoinPolicy));
+		UE_LOG(LogTemp, Log, TEXT("	NumMaxConnections [%d]"), SessionSettings.NumMaxConnections);
+		UE_LOG(LogTemp, Log, TEXT("	SchemaName [%s]"), *SessionSettings.SchemaName.ToString());
+		UE_LOG(LogTemp, Log, TEXT("	CustomSettings:"));
+
+		for (const TPair<FSchemaAttributeId, FCustomSessionSetting>& Entry : SessionSettings.CustomSettings)
+		{
+			UE_LOG(LogTemp, Log, TEXT("		Key: [%s]"), *Entry.Key.ToString());
+			UE_LOG(LogTemp, Log, TEXT("		Value:	Visibility [%s]"), LexToString(Entry.Value.Visibility));
+			UE_LOG(LogTemp, Log, TEXT("				ID [%d]"), Entry.Value.ID);
+			UE_LOG(LogTemp, Log, TEXT("				Data [%s]"), *LexToString(Entry.Value.Data));
+		}
+	}
+
 	/** FSessionsCommon */
+
 	FSessionsCommon::FSessionsCommon(FOnlineServicesCommon& InServices)
 		: TOnlineComponent(TEXT("Sessions"), InServices)
 	{
@@ -62,7 +117,8 @@ namespace UE::Online {
 		RegisterCommand(&FSessionsCommon::AddSessionMember);
 		RegisterCommand(&FSessionsCommon::RemoveSessionMember);
 		RegisterCommand(&FSessionsCommon::SendSessionInvite);
-		RegisterCommand(&FSessionsCommon::GetSessionInvites);
+		RegisterCommand(&FSessionsCommon::GetSessionInviteById);
+		RegisterCommand(&FSessionsCommon::GetAllSessionInvites);
 		RegisterCommand(&FSessionsCommon::RejectSessionInvite);
 	}
 
@@ -70,24 +126,20 @@ namespace UE::Online {
 	{
 		// TODO: Params and user login check
 
+		TArray<TSharedRef<const ISession>> Sessions;
+
 		if (const TArray<FName>* UserSessions = NamedSessionUserMap.Find(Params.LocalAccountId))
 		{
-			FGetAllSessions::Result Result;
-
 			for (const FName& SessionName : *UserSessions)
 			{
 				const FOnlineSessionId& SessionId = LocalSessionsByName.FindChecked(SessionName);
 				const TSharedRef<FSessionCommon>& Session = AllSessionsById.FindChecked(SessionId);
 
-				Result.Sessions.Add(Session);
+				Sessions.Add(Session);
 			}
-
-			return TOnlineResult<FGetAllSessions>(MoveTemp(Result));
-		}
-		else
-		{
-			return TOnlineResult<FGetAllSessions>(Errors::NotFound());
 		}
+
+		return TOnlineResult<FGetAllSessions>({ MoveTemp(Sessions) });
 	}
 
 	TOnlineResult<FGetSessionByName> FSessionsCommon::GetSessionByName(FGetSessionByName::Params&& Params) const
@@ -310,7 +362,26 @@ namespace UE::Online {
 		return Operation->GetHandle();
 	}
 
-	TOnlineResult<FGetSessionInvites> FSessionsCommon::GetSessionInvites(FGetSessionInvites::Params&& Params)
+	TOnlineResult<FGetSessionInviteById> FSessionsCommon::GetSessionInviteById(FGetSessionInviteById::Params&& Params)
+	{
+		if (const TMap<FSessionInviteId, TSharedRef<FSessionInvite>>* UserMap = SessionInvitesUserMap.Find(Params.LocalAccountId))
+		{
+			if (const TSharedRef<FSessionInvite>* SessionInvite = UserMap->Find(Params.SessionInviteId))
+			{
+				return TOnlineResult<FGetSessionInviteById>({ *SessionInvite });
+			}
+			else
+			{
+				return TOnlineResult<FGetSessionInviteById>(Errors::NotFound());
+			}
+		}
+		else
+		{
+			return TOnlineResult<FGetSessionInviteById>(Errors::InvalidState());
+		}		
+	}
+
+	TOnlineResult<FGetAllSessionInvites> FSessionsCommon::GetAllSessionInvites(FGetAllSessionInvites::Params&& Params)
 	{
 		TArray<TSharedRef<const FSessionInvite>> SessionInvites;
 
@@ -321,12 +392,8 @@ namespace UE::Online {
 				SessionInvites.Add(Entry.Value);
 			}
 		}
-		else
-		{
-			return TOnlineResult<FGetSessionInvites>(Errors::NotFound());
-		}
 
-		return TOnlineResult<FGetSessionInvites>({ SessionInvites });
+		return TOnlineResult<FGetAllSessionInvites>({ MoveTemp(SessionInvites) });
 	}
 
 	TOnlineAsyncOpHandle<FRejectSessionInvite> FSessionsCommon::RejectSessionInvite(FRejectSessionInvite::Params&& Params)
diff --git a/Engine/Plugins/Online/Experimental/OnlineServices/Source/OnlineServicesCommon/Public/Online/SessionsCommon.h b/Engine/Plugins/Online/Experimental/OnlineServices/Source/OnlineServicesCommon/Public/Online/SessionsCommon.h
--- a/Engine/Plugins/Online/Experimental/OnlineServices/Source/OnlineServicesCommon/Public/Online/SessionsCommon.h
+++ b/Engine/Plugins/Online/Experimental/OnlineServices/Source/OnlineServicesCommon/Public/Online/SessionsCommon.h
@@ -140,11 +140,18 @@ public:
 	virtual const uint32 GetNumOpenConnections() const override				{ return SessionSettings.NumMaxConnections - SessionMembers.Num(); }
 	virtual const FSessionInfo& GetSessionInfo() const override				{ return SessionInfo; }
 	virtual const FSessionSettings GetSessionSettings() const override		{ return SessionSettings; }
-	virtual const FSessionMembersSet& GetSessionMembers() const override	{ return SessionMembers; }
+	virtual const FSessionMemberIdsSet& GetSessionMembers() const override	{ return SessionMembers; }
 
 	virtual bool IsJoinable() const override								{ return GetNumOpenConnections() > 0 && SessionSettings.bAllowNewMembers; }
 
-	virtual FString ToLogString() const override							{ return FString::Printf(TEXT("[FSessionCommon: SessionId[%d]]"), SessionInfo.SessionId.GetHandle()); } // TODO: Expand with relevant information
+	virtual FString ToLogString() const override;
+
+	virtual void DumpState() const override;
+
+protected:
+	void DumpMemberData() const;
+	void DumpSessionInfo() const;
+	void DumpSessionSettings() const;
 
 public:
 	/** Session information that will remain constant throughout the session's lifetime */
@@ -159,7 +166,7 @@ public:
 	FSessionSettings SessionSettings;
 
 	/* Set containing user ids for all the session members */
-	FSessionMembersSet SessionMembers;
+	FSessionMemberIdsSet SessionMembers;
 
 	FSessionCommon& operator+=(const FSessionUpdate& SessionUpdate);
 };
@@ -248,7 +255,8 @@ public:
 	virtual TOnlineAsyncOpHandle<FAddSessionMember> AddSessionMember(FAddSessionMember::Params&& Params) override;
 	virtual TOnlineAsyncOpHandle<FRemoveSessionMember> RemoveSessionMember(FRemoveSessionMember::Params&& Params) override;
 	virtual TOnlineAsyncOpHandle<FSendSessionInvite> SendSessionInvite(FSendSessionInvite::Params&& Params) override;
-	virtual TOnlineResult<FGetSessionInvites> GetSessionInvites(FGetSessionInvites::Params&& Params) override;
+	virtual TOnlineResult<FGetSessionInviteById> GetSessionInviteById(FGetSessionInviteById::Params&& Params) override;
+	virtual TOnlineResult<FGetAllSessionInvites> GetAllSessionInvites(FGetAllSessionInvites::Params&& Params) override;
 	virtual TOnlineAsyncOpHandle<FRejectSessionInvite> RejectSessionInvite(FRejectSessionInvite::Params&& Params) override;
 
 	virtual TOnlineEvent<void(const FSessionJoined&)> OnSessionJoined() override;
diff --git a/Engine/Plugins/Online/Experimental/OnlineServices/Source/OnlineServicesCommonEngineUtils/Private/Online/SessionsLAN.cpp b/Engine/Plugins/Online/Experimental/OnlineServices/Source/OnlineServicesCommonEngineUtils/Private/Online/SessionsLAN.cpp
--- a/Engine/Plugins/Online/Experimental/OnlineServices/Source/OnlineServicesCommonEngineUtils/Private/Online/SessionsLAN.cpp
+++ b/Engine/Plugins/Online/Experimental/OnlineServices/Source/OnlineServicesCommonEngineUtils/Private/Online/SessionsLAN.cpp
@@ -99,6 +99,12 @@ const FSessionLAN& FSessionLAN::Cast(const ISession& InSession)
 	return static_cast<const FSessionLAN&>(InSession);
 }
 
+void FSessionLAN::DumpState() const
+{
+	FSessionCommon::DumpState();
+	UE_LOG(LogTemp, Log, TEXT("OwnerInternetAddr [%s]"), *OwnerInternetAddr->ToString(true));
+}
+
 /** FSessionsLAN */
 
 FSessionsLAN::FSessionsLAN(FOnlineServicesCommon& InServices)
diff --git a/Engine/Plugins/Online/Experimental/OnlineServices/Source/OnlineServicesCommonEngineUtils/Public/Online/SessionsLAN.h b/Engine/Plugins/Online/Experimental/OnlineServices/Source/OnlineServicesCommonEngineUtils/Public/Online/SessionsLAN.h
--- a/Engine/Plugins/Online/Experimental/OnlineServices/Source/OnlineServicesCommonEngineUtils/Public/Online/SessionsLAN.h
+++ b/Engine/Plugins/Online/Experimental/OnlineServices/Source/OnlineServicesCommonEngineUtils/Public/Online/SessionsLAN.h
@@ -30,6 +30,8 @@ public:
 	static FSessionLAN& Cast(FSessionCommon& InSession);
 	static const FSessionLAN& Cast(const ISession& InSession);
 
+	virtual void DumpState() const override;
+
 private:
 	void Initialize();
 
diff --git a/Engine/Plugins/Online/Experimental/OnlineServices/Source/OnlineServicesInterface/Public/Online/Sessions.h b/Engine/Plugins/Online/Experimental/OnlineServices/Source/OnlineServicesInterface/Public/Online/Sessions.h
--- a/Engine/Plugins/Online/Experimental/OnlineServices/Source/OnlineServicesInterface/Public/Online/Sessions.h
+++ b/Engine/Plugins/Online/Experimental/OnlineServices/Source/OnlineServicesInterface/Public/Online/Sessions.h
@@ -10,6 +10,9 @@
 
 namespace UE::Online {
 
+// ISessions Types
+
+/** Filter class for a session search, will be compared against the session's CustomSettings */
 struct FFindSessionsSearchFilter
 {
 	/** Name of the custom setting to be used as filter */
@@ -22,14 +25,13 @@ struct FFindSessionsSearchFilter
 	FSchemaVariant Value;
 };
 
-// Custom Session Settings
-
+/** User-defined data to be stored along with the session information. Will be evaluated in session searches */
 struct FCustomSessionSetting
 {
-	/** Setting value */
+	/** Setting data value */
 	FSchemaVariant Data;
 
-	/** How is this session setting advertised with the backend or searches */
+	/** Manner in which this session setting is advertised with the backend or searches */
 	ESchemaAttributeVisibility Visibility;
 
 	/** Optional ID used in some platforms as the index instead of the setting name */
@@ -38,27 +40,36 @@ struct FCustomSessionSetting
 
 using FCustomSessionSettingsMap = TMap<FSchemaAttributeId, FCustomSessionSetting>;
 
+/** Contains both the old and new values for an updated FCustomSessionSetting. Part of the FSessionUpdate event data */
 struct FCustomSessionSettingUpdate
 {
+	/** Old value for the updated Setting */
 	FCustomSessionSetting OldValue;
 
+	/** New value for the updated Setting */
 	FCustomSessionSetting NewValue;
 };
 
 using FCustomSessionSettingUpdateMap = TMap<FName, FCustomSessionSettingUpdate>;
 
-using FSessionMembersSet = TSet<FAccountId>;
+using FSessionMemberIdsSet = TSet<FAccountId>;
 
+/** Set of options to reflect how a session may be discovered in searches and joined */
 enum class ESessionJoinPolicy : uint8
 {
+	/** The session will appear on searches an may be joined by anyone */
 	Public,
+
+	/** The session will not appear on searches and may only be joined via presence (if enabled) or invitation */
 	FriendsOnly,
+
+	/** The session will not appear on searches and may not be joined via presence, only via invitation */
 	InviteOnly
 };
 ONLINESERVICESINTERFACE_API const TCHAR* LexToString(ESessionJoinPolicy Value);
 ONLINESERVICESINTERFACE_API void LexFromString(ESessionJoinPolicy& Value, const TCHAR* InStr);
 
-/** Contains new values for an FSessions modifiable settings. Taken as a parameter by FUpdateSessions method */
+/** Contains new values for an FSessions modifiable settings. Taken as a parameter by FUpdateSessionSettings method */
 struct ONLINESERVICESINTERFACE_API FSessionSettingsUpdate
 {
 	/** Set with an updated value if the SchemaName field will be changed in the update operation */
@@ -70,40 +81,40 @@ struct ONLINESERVICESINTERFACE_API FSessionSettingsUpdate
 	/** Set with an updated value if the bAllowNewMembers field will be changed in the update operation */
 	TOptional<bool> bAllowNewMembers;
 
-	/** Updated values for custom settings to change in the update operation*/
+	/** Updated values for custom settings to change in the update operation */
 	FCustomSessionSettingsMap UpdatedCustomSettings;
-	/** Names of custom settings to be removed in the update operation*/
+	/** Names of custom settings to be removed in the update operation */
 	TArray<FSchemaAttributeId> RemovedCustomSettings;
 
 	FSessionSettingsUpdate& operator+=(FSessionSettingsUpdate&& UpdatedValue);
 };
 
-/** Contains updated data for any modifiable members of FSessionSettings. Member of FSessionUpdated event */
+/** Contains updated data for any modifiable members of FSessionSettings. Part of the FSessionUpdated event data */
 struct FSessionSettingsChanges
 {
-	/* If set, the FSessionSettings's SchemaName member will be updated to this value */
+	/* If set, the FSessionSettings's SchemaName member will have been updated to this value */
 	TOptional<FName> SchemaName;
-	/** If set, the FSessionSettings's NumMaxConnections member will be updated to this value */
+	/** If set, the FSessionSettings's NumMaxConnections member will have been updated to this value */
 	TOptional<uint32> NumMaxConnections;
-	/** If set, the FSessionSettings's JoinPolicy member will be updated to this value */
+	/** If set, the FSessionSettings's JoinPolicy member will have been updated to this value */
 	TOptional<ESessionJoinPolicy> JoinPolicy;
-	/** If set, the FSessionSettings's bAllowNewMembers member will be updated to this value */
+	/** If set, the FSessionSettings's bAllowNewMembers member will have been updated to this value */
 	TOptional<bool> bAllowNewMembers;
 
-	/** New custom settings, with their values */
+	/** New custom settings added in the update, with their values */
 	FCustomSessionSettingsMap AddedCustomSettings;
 
 	/** Existing custom settings that changed value, including new and old values */
 	FCustomSessionSettingUpdateMap ChangedCustomSettings;
 
-	/** Keys for removed custom settings */
+	/** Keys for custom settings removed in the update */
 	TArray<FName> RemovedCustomSettings;
 };
 
-/** Set of all of an FSession's defining properties that can be updated by the session owner during its lifetime */
+/** Set of all of an FSession's defining properties that can be updated by the session owner during its lifetime, using the FUpdateSessionSettings method */
 struct ONLINESERVICESINTERFACE_API FSessionSettings
 {
-	/* The schema which will be applied to the session */
+	/* The name for the schema which will be applied to the session's user-defined attributes */
 	FName SchemaName;
 
 	/* Maximum number of slots for session members */
@@ -124,10 +135,10 @@ struct ONLINESERVICESINTERFACE_API FSessionSettings
 /** Information about an FSession that will be set at creation time and remain constant during its lifetime */
 struct FSessionInfo
 {
-	/** The id for the session, platform dependent */
+	/** The id handle for the session, platform dependent */
 	FOnlineSessionId SessionId;
 
-	/* In platforms that support this feature, it will set the session id to this value. Might be subject to minimum and maximum length */
+	/* In platforms that support this feature, it will set the session id to this value. Might be subject to minimum and maximum length restrictions */
 	FString SessionIdOverride;
 
 	/* Whether the session is only available in the local network and not via internet connection. Only available in some platforms. False by default */
@@ -143,51 +154,72 @@ struct FSessionInfo
 	bool bAntiCheatProtected = false;
 };
 
+/** Interface to access all information related to Online Sessions. Read only */
 class ISession
 {
 public:
+	/** Retrieves the id handle for the user who created or currently owns the session */
 	virtual const FAccountId GetOwnerAccountId() const = 0;
+
+	/** Retrieves the id handle for the session */
 	virtual const FOnlineSessionId GetSessionId() const = 0;
+
+	/** Retrieves the number of available slots for new session members */
 	virtual const uint32 GetNumOpenConnections() const = 0;
+
+	/** Retrieves the set of constant information about the session */
 	virtual const FSessionInfo& GetSessionInfo() const = 0;
+
+	/** Retrieves the set of variable information about the session */
 	virtual const FSessionSettings GetSessionSettings() const = 0;
-	virtual const FSessionMembersSet& GetSessionMembers() const = 0;
+
+	/** Retrieves the list of users currently in the session */
+	virtual const FSessionMemberIdsSet& GetSessionMembers() const = 0;
 
 	/** Evaluates a series of factors to determine if a session is accepting new members */
 	virtual bool IsJoinable() const = 0;
 
+	/** Returns a string with the minimal information to identify the session */
 	virtual FString ToLogString() const = 0;
+
+	/** Returns a string with all the information in the session */
+	virtual void DumpState() const = 0;
 };
 ONLINESERVICESINTERFACE_API const FString ToLogString(const ISession& Session);
 
+/** Information about an invitation to join an Online Session */
 struct FSessionInvite
 {
-	/* The user which the invite got sent to */
+	/* The id handle for the user which the invite got sent to */
 	FAccountId RecipientId;
 
-	/* The user which sent the invite */
+	/* The id handle for the user which sent the invite */
 	FAccountId SenderId;
 
-	/* The invite id handle, needed for retrieving session information and rejecting the invite */
+	/* The invite id handle, needed for retrieving invite information and rejecting the invite */
 	FSessionInviteId InviteId;
 
-	/* Pointer to the session information */
+	/* The session id handle, needed for retrieving the session information */
 	FOnlineSessionId SessionId;
 
-	// TODO: Default constructor will be deleted after we cache invites
+	// TODO: Refactor into ISessionInvite, FSessionInviteCommon/EOSGS/OSSAdapter to add a ToLogString method
 };
 
+// ISessions Methods
+
 struct FGetAllSessions
 {
 	static constexpr TCHAR Name[] = TEXT("GetAllSessions");
 
 	struct Params
 	{
+		/** Id handle for the local user which will perform the action */
 		FAccountId LocalAccountId;
 	};
 
 	struct Result
 	{
+		/** Array of sessions that the given user is member of */
 		TArray<TSharedRef<const ISession>> Sessions;
 	};
 };
@@ -198,11 +230,13 @@ struct FGetSessionByName
 
 	struct Params
 	{
+		/** Local name for the session */
 		FName LocalName;
 	};
 
 	struct Result
 	{
+		/** Reference to the session mapped to the given name */
 		TSharedRef<const ISession> Session;
 
 		Result() = delete; // cannot default construct due to TSharedRef
@@ -215,11 +249,13 @@ struct FGetSessionById
 
 	struct Params
 	{
+		/** Id handle for the session to be retrieved */
 		FOnlineSessionId SessionId;
 	};
 
 	struct Result
 	{
+		/** Reference to the session mapped to the given id */
 		TSharedRef<const ISession> Session;
 
 		Result() = delete; // cannot default construct due to TSharedRef
@@ -232,11 +268,13 @@ struct FGetPresenceSession
 
 	struct Params
 	{
+		/** Id handle for the local user which will perform the action */
 		FAccountId LocalAccountId;
 	};
 
 	struct Result
 	{
+		/** Reference to the session set as presence session for the given user */
 		TSharedRef<const ISession> Session;
 	};
 };
@@ -247,13 +285,16 @@ struct FIsPresenceSession
 
 	struct Params
 	{
+		/** Id handle for the local user which will perform the action */
 		FAccountId LocalAccountId;
 
+		/** Id handle for the session to be compared */
 		FOnlineSessionId SessionId;
 	};
 
 	struct Result
 	{
+		/** Whether the session mapped to the given id is set as the presence session for the given user */
 		bool bIsPresenceSession;
 	};
 };
@@ -264,8 +305,10 @@ struct FSetPresenceSession
 
 	struct Params
 	{
+		/** Id handle for the local user which will perform the action */
 		FAccountId LocalAccountId;
 
+		/** Id handle for the session */
 		FOnlineSessionId SessionId;
 	};
 
@@ -280,6 +323,7 @@ struct FClearPresenceSession
 
 	struct Params
 	{
+		/** Id handle for the local user which will perform the action */
 		FAccountId LocalAccountId;
 	};
 
@@ -294,28 +338,28 @@ struct FCreateSession
 
 	struct Params
 	{
-		/** The local user agent which will perform the action. */
+		/** Id handle for the local user which will perform the action */
 		FAccountId LocalAccountId;
 
-		/** The local name for the session */
+		/** Local name for the session */
 		FName SessionName;
 
-		/* In platforms that support this feature, it will set the session id to this value. Might be subject to minimum and maximum length */
+		/** In platforms that support this feature, it will set the session id to this value. Might be subject to minimum and maximum length restrictions */
 		FString SessionIdOverride;
 
-		/** Whether this session should be set as the user's new presence session. False by default */
+		/** Whether this session should be set as the local user's new presence session. False by default */
 		bool bPresenceEnabled = false;
 
-		/* Whether the session is only available in the local network and not via internet connection. Only available in some platforms. False by default */
+		/** Whether the session is only available via the local network and not via internet connection. Only available in some platforms. False by default */
 		bool bIsLANSession = false;
 
-		/* Whether the session is configured to run as a dedicated server. Only available in some platforms. False by default */
+		/** Whether the session is configured to run from a dedicated server. Only available in some platforms. False by default */
 		bool bIsDedicatedServerSession = false;
 
-		/* Whether this session will allow sanctioned players to join it. True by default */
+		/** Whether this session will allow sanctioned players to join it. True by default */
 		bool bAllowSanctionedPlayers = true;
 
-		/* Whether this is a secure session protected by anti-cheat services. False by default */
+		/** Whether this is a secure session protected by anti-cheat services. False by default */
 		bool bAntiCheatProtected = false;
 
 		/** Settings object to define session properties during creation */
@@ -334,10 +378,10 @@ struct FUpdateSessionSettings
 
 	struct Params
 	{
-		/** The local user agent which will perform the action. */
+		/** Id handle for the local user which will perform the action */
 		FAccountId LocalAccountId;
 
-		/** The local name for the session */
+		/** Local name for the session */
 		FName SessionName;
 
 		/** Changes to current session settings */
@@ -356,13 +400,13 @@ struct FLeaveSession
 
 	struct Params
 	{
-		/* The local user agent which leaves the session */
+		/** Id handle for the local user which will perform the action */
 		FAccountId LocalAccountId;
 
-		/* The local name for the session. */
+		/** Local name for the session */
 		FName SessionName;
 
-		/* Whether the call should attempt to destroy the session instead of just leave it */
+		/** Whether the call should attempt to destroy the session instead of just leave it */
 		bool bDestroySession;
 	};
 
@@ -378,22 +422,22 @@ struct FFindSessions
 
 	struct Params
 	{
-		/* The local user agent which starts the session search*/
+		/** Id handle for the local user which will perform the action */
 		FAccountId LocalAccountId;
 
-		/* Maximum number of results to return in one search */
+		/** Maximum number of results to return in the search */
 		uint32 MaxResults;
 
-		/** Whether we want to look for LAN sessions or Online sessions */
-		bool bFindLANSessions;
+		/** Whether we want to look for LAN sessions (true) or Online sessions (false). False by default */
+		bool bFindLANSessions = false;
 
-		/* Filters to apply when searching for sessions. */
+		/** Filters to apply when searching for sessions. */
 		TArray<FFindSessionsSearchFilter> Filters;
 
-		/* Find sessions containing the target user. */
+		/** If set, the search will look for sessions containing the set user. */
 		TOptional<FAccountId> TargetUser;
 
-		/* Find join info for the target session id. */
+		/** If set, the search will look for the session with the set session id. */
 		TOptional<FOnlineSessionId> SessionId;
 	};
 
@@ -409,10 +453,10 @@ struct FStartMatchmaking
 
 	struct Params
 	{
-		/* Session creation parameters */
+		/** Session creation parameters */
 		FCreateSession::Params SessionCreationParameters;
 
-		/* Filters to apply when searching for sessions */
+		/** Filters to apply when searching for sessions */
 		TArray<FFindSessionsSearchFilter> SessionSearchFilters;
 	};
 
@@ -428,16 +472,16 @@ struct FJoinSession
 
 	struct Params
 	{
-		/* The local user agent which starts the join operation*/
+		/** Id handle for the local user which will perform the action */
 		FAccountId LocalAccountId;
 
-		/* Local name for the session */
+		/** Local name for the session */
 		FName SessionName;
 
-		/* Id handle for the session to be joined. To be retrieved via session search or invite */
+		/** Id handle for the session to be joined. To be retrieved via session search, invite or UI operation */
 		FOnlineSessionId SessionId;
 
-		/* Whether this session should be set as the user's new presence session. False by default */
+		/** Whether this session should be set as the user's new presence session. False by default */
 		bool bPresenceEnabled = false;
 	};
 
@@ -453,10 +497,10 @@ struct FAddSessionMember
 
 	struct Params
 	{
-		/* The local user agent */
+		/** Id handle for the local user which will perform the action */
 		FAccountId LocalAccountId;
 
-		/* Local name for the session */
+		/** Local name for the session */
 		FName SessionName;
 	};
 
@@ -472,10 +516,10 @@ struct FRemoveSessionMember
 
 	struct Params
 	{
-		/* The local user agent */
+		/** Id handle for the local user which will perform the action */
 		FAccountId LocalAccountId;
 
-		/* Local name for the session */
+		/** Local name for the session */
 		FName SessionName;
 	};
 
@@ -491,13 +535,13 @@ struct FSendSessionInvite
 
 	struct Params
 	{
-		/* The local user agent which sends the invite*/
+		/** Id handle for the local user which will perform the action */
 		FAccountId LocalAccountId;
 
-		/* The local name for the session. */
+		/** Local name for the session */
 		FName SessionName;
 
-		/* Array of id handles for users to which the invites will be sent */
+		/** Array of id handles for users to which the invites will be sent */
 		TArray<FAccountId> TargetUsers;
 	};
 
@@ -507,19 +551,39 @@ struct FSendSessionInvite
 	};
 };
 
-struct FGetSessionInvites
+struct FGetSessionInviteById
 {
-	static constexpr TCHAR Name[] = TEXT("GetSessionInvites");
+	static constexpr TCHAR Name[] = TEXT("GetSessionInviteById");
 
 	struct Params
 	{
-		/* The local user agent */
+		/** Id handle for the local user which will perform the action */
 		FAccountId LocalAccountId;
+
+		/** Id handle for the invite to be retrieved */
+		FSessionInviteId SessionInviteId;
 	};
 
 	struct Result
 	{
-		/** Set of active session invites */
+		/** Reference to the invite mapped to the given id */
+		TSharedRef<const FSessionInvite> SessionInvite;
+	};
+};
+
+struct FGetAllSessionInvites
+{
+	static constexpr TCHAR Name[] = TEXT("GetAllSessionInvites");
+
+	struct Params
+	{
+		/** Id handle for the local user which will perform the action */
+		FAccountId LocalAccountId;
+	};
+
+	struct Result
+	{
+		/** Array of invites received by the given user */
 		TArray<TSharedRef<const FSessionInvite>> SessionInvites;
 	};
 };
@@ -530,10 +594,10 @@ struct FRejectSessionInvite
 
 	struct Params
 	{
-		/* The local user agent which started the query*/
+		/** Id handle for the local user which will perform the action */
 		FAccountId LocalAccountId;
 
-		/* The id handle for the invite to be rejected */
+		/* Id handle for the session invite to be rejected */
 		FSessionInviteId SessionInviteId;
 	};
 
@@ -546,16 +610,16 @@ struct FRejectSessionInvite
 
 struct FSessionJoined
 {
-	/* The local user which joined the session */
+	/** Id handle for the local user who joined the session */
 	FAccountId LocalAccountId;
 
-	/* Id for the session joined. */
+	/** Id handle for the session joined. */
 	FOnlineSessionId SessionId;
 };
 
 struct FSessionLeft
 {
-	/* The local users which left the session */
+	/** Id handle for the local users which left the session */
 	FAccountId LocalAccountId;
 };
 
@@ -568,36 +632,33 @@ struct ONLINESERVICESINTERFACE_API FSessionUpdate
 	/** If set, the SessionSettings member will have updated using the struct information */
 	TOptional<FSessionSettingsChanges> SessionSettingsChanges;
 
-	/** Session member information for members that just joined the session */
-	FSessionMembersSet AddedSessionMembers;
+	/** Id handles for members that just joined the session */
+	FSessionMemberIdsSet AddedSessionMembers;
 
 	/** Id handles for members that just left the session */
-	FSessionMembersSet RemovedSessionMembers;
+	FSessionMemberIdsSet RemovedSessionMembers;
 
 	FSessionUpdate& operator+=(const FSessionUpdate& SessionUpdate);
 };
 
 struct FSessionUpdated
 {
-	/* Name for the session updated */
+	/** Local name for the updated session */
 	FName SessionName;
 
-	/* Updated session settings */
+	/** Updated session information */
 	FSessionUpdate SessionUpdate;
 };
 
 struct FSessionInviteReceived
 {
-	/* The local user which received the invite */
+	/** Id handle for the local user which received the invite */
 	FAccountId LocalAccountId;
 
-	/** The session invite the local user was sent, or the online error if there was a failure retrieving the session for it*/
-	TSharedRef<const FSessionInvite> SessionInvite;
-
-	FSessionInviteReceived() = delete; // cannot default construct due to TSharedRef. TODO: Add a GetSessionInvite method and return an id here
+	/** The session invite id for the invite received */
+	FSessionInviteId SessionInviteId;
 };
 
-/** Session join requested source */
 enum class EUISessionJoinRequestedSource : uint8
 {
 	/** Unspecified by the online service */
@@ -610,10 +671,10 @@ ONLINESERVICESINTERFACE_API void LexFromString(EUISessionJoinRequestedSource& Ou
 
 struct FUISessionJoinRequested
 {
-	/** The local user associated with the join request. */
+	/** Id handle for the local user associated with the join request */
 	FAccountId LocalAccountId;
 
-	/** The id for the session the local user requested to join, or the online error if there was a failure retrieving it */
+	/** Id handle for the session the local user requested to join, or the online error if there was a failure retrieving it */
 	TResult<FOnlineSessionId, FOnlineError> Result;
 
 	/** Join request source */
@@ -634,7 +695,7 @@ public:
 	virtual TOnlineResult<FGetAllSessions> GetAllSessions(FGetAllSessions::Params&& Params) const = 0;
 
 	/**
-	 * Get the session object with a given local name.
+	 * Gets a reference to the session with the given local name.
 	 *
 	 * @params Parameters for the GetSessionByName call
 	 * return
@@ -642,7 +703,7 @@ public:
 	virtual TOnlineResult<FGetSessionByName> GetSessionByName(FGetSessionByName::Params&& Params) const = 0;
 
 	/**
-	 * Get the session object with a given id handle.
+	 * Gets a reference to the session with the given id handle.
 	 *
 	 * @params Parameters for the GetSessionById call
 	 * return
@@ -650,7 +711,7 @@ public:
 	virtual TOnlineResult<FGetSessionById> GetSessionById(FGetSessionById::Params&& Params) const = 0;
 
 	/**
-	 * Get the session set as presence session for the user.
+	 * Gets a reference to the session set as presence session for the given user.
 	 *
 	 * @params Parameters for the GetPresenceSession call
 	 * return
@@ -658,7 +719,7 @@ public:
 	virtual TOnlineResult<FGetPresenceSession> GetPresenceSession(FGetPresenceSession::Params&& Params) const = 0;
 
 	/**
-	 * Returns whether the session with the given id is set as the presence session for the user.
+	 * Returns whether the session with the given id is set as the presence session for the given user.
 	 *
 	 * @params Parameters for the IsPresenceSession call
 	 * return
@@ -666,7 +727,7 @@ public:
 	virtual TOnlineResult<FIsPresenceSession> IsPresenceSession(FIsPresenceSession::Params&& Params) const = 0;
 
 	/**
-	 * Sets the session with the given id as the presence session for the user.
+	 * Sets the session with the given id as the presence session for the given user.
 	 *
 	 * @params Parameters for the SetPresenceSession call
 	 * return
@@ -674,7 +735,7 @@ public:
 	virtual TOnlineResult<FSetPresenceSession> SetPresenceSession(FSetPresenceSession::Params&& Params) = 0;
 
 	/**
-	 * Clears the presence session for the user. If no presence session is set, GetPresenceSession will return an error.
+	 * Clears the presence session for the given user.
 	 *
 	 * @params Parameters for the ClearPresenceSession call
 	 * return
@@ -682,7 +743,8 @@ public:
 	virtual TOnlineResult<FClearPresenceSession> ClearPresenceSession(FClearPresenceSession::Params&& Params) = 0;
 
 	/**
-	 * Create and join a new session.
+	 * Creates a new session with the given parameters, and assigns to it the given local name.
+	 * Depending on the implementation, the creating user might not be added to the Session Members automatically, so a subsequent call to AddSessionMember is recommended
 	 *
 	 * @param Parameters for the CreateSession call
 	 * @return
@@ -690,7 +752,8 @@ public:
 	virtual TOnlineAsyncOpHandle<FCreateSession> CreateSession(FCreateSession::Params&& Params) = 0;
 
 	/**
-	 * Update a given session's settings. Can only be called by the session owner.
+	 * Update the settings for the session with the given name.
+	 * Should only be called by the session owner.
 	 *
 	 * @param Parameters for the UpdateSessionSettings call
 	 * @return
@@ -698,7 +761,7 @@ public:
 	virtual TOnlineAsyncOpHandle<FUpdateSessionSettings> UpdateSessionSettings(FUpdateSessionSettings::Params&& Params) = 0;
 
 	/**
-	 * Leave and optionally destroy a given session.
+	 * Leaves and optionally destroys the session with the given name.
 	 *
 	 * @param Parameters for the LeaveSession call
 	 * @return
@@ -714,7 +777,7 @@ public:
 	virtual TOnlineAsyncOpHandle<FFindSessions> FindSessions(FFindSessions::Params&& Params) = 0;
 
 	/**
-	 * Starts the matchmaking process, which will either create a session with the passed parameters, or join one that matches the passed search filters.
+	 * Starts the matchmaking process, which will either create a session with the given parameters, or join one that matches the given search filters.
 	 *
 	 * @param Parameters for the StartMatchmaking call
 	 * @return
@@ -722,7 +785,7 @@ public:
 	virtual TOnlineAsyncOpHandle<FStartMatchmaking> StartMatchmaking(FStartMatchmaking::Params&& Params) = 0;
 
 	/**
-	 * Starts the join process for the given session for all users provided.
+	 * Joins  the session with the given session id, and assigns to it the given local name.
 	 *
 	 * @param Parameters for the JoinSession call
 	 * @return
@@ -730,9 +793,8 @@ public:
 	virtual TOnlineAsyncOpHandle<FJoinSession> JoinSession(FJoinSession::Params&& Params) = 0;
 
 	/**
-	 * Adds a set of new session members to the named session
-	 * Session member information passed will be saved in the session settings
-	 * Number of open slots in the session will decrease accordingly
+	 * Adds the given user as a new session member to the session with the given name
+	 * The number of open slots in the session will decrease accordingly
 	 * 
 	 * @params Parameters for the AddSessionMember call
 	 * @return
@@ -740,9 +802,8 @@ public:
 	virtual TOnlineAsyncOpHandle<FAddSessionMember> AddSessionMember(FAddSessionMember::Params&& Params) = 0;
 
 	/**
-	 * Removes a set of session member from the named session
-	 * Session member information for them will be removed from session settings
-	 * Number of open slots in the session will increase accordingly
+	 * Removes the given user from the session with the given name
+	 * The number of open slots in the session will increase accordingly
 	 *
 	 * @params Parameters for the RemoveSessionMember call
 	 * @return
@@ -750,7 +811,7 @@ public:
 	virtual TOnlineAsyncOpHandle<FRemoveSessionMember> RemoveSessionMember(FRemoveSessionMember::Params&& Params) = 0;
 
 	/**
-	 * Sends an invite to the named session to all given users.
+	 * Sends an invite to the session with the given name to all given users.
 	 *
 	 * @param Parameters for the SendSessionInvite call
 	 * @return
@@ -758,15 +819,23 @@ public:
 	virtual TOnlineAsyncOpHandle<FSendSessionInvite> SendSessionInvite(FSendSessionInvite::Params&& Params) = 0;
 
 	/**
-	 * Returns all cached session invites for the given user.
+	 * Gets a reference to the session invite with the given invite id.
+	 *
+	 * @param Parameters for the GetSessionInviteById call
+	 * @return
+	 */
+	virtual TOnlineResult<FGetSessionInviteById> GetSessionInviteById(FGetSessionInviteById::Params&& Params) = 0;
+
+	/**
+	 * Gets an array of references to all the session invites the given user has received.
 	 *
 	 * @param Parameters for the SendSessionInvite call
 	 * @return
 	 */
-	virtual TOnlineResult<FGetSessionInvites> GetSessionInvites(FGetSessionInvites::Params&& Params) = 0;
+	virtual TOnlineResult<FGetAllSessionInvites> GetAllSessionInvites(FGetAllSessionInvites::Params&& Params) = 0;
 
 	/**
-	 * Rejects a given session invite for a user.
+	 * Rejects the session invite with the given invite id.
 	 *
 	 * @param Parameters for the RejectSessionInvite call
 	 * @return
@@ -776,15 +845,13 @@ public:
 	/* Events */
 
 	/**
-	 * Get the event that is triggered when a session is joined.
-	 * This event will trigger as a result of creating or joining a session.
+	 * This event will trigger as a result of joining a session.
 	 *
 	 * @return
 	 */
 	virtual TOnlineEvent<void(const FSessionJoined&)> OnSessionJoined() = 0;
 
 	/**
-	 * Get the event that is triggered when a session is left.
 	 * This event will trigger as a result of leaving or destroying a session.
 	 *
 	 * @return
@@ -792,24 +859,21 @@ public:
 	virtual TOnlineEvent<void(const FSessionLeft&)> OnSessionLeft() = 0;
 
 	/**
-	 * Get the event that is triggered when a session invite is accepted.
-	 * This event will trigger as a result of accepting a platform session invite.
+	 * This event will trigger as a result of updating a session's settings, or whenever a session update event is received by the API.
 	 *
 	 * @return
 	 */
 	virtual TOnlineEvent<void(const FSessionUpdated&)> OnSessionUpdated() = 0;
 
 	/**
-	 * Get the event that is triggered when a session invite is received.
-	 * This event will trigger as a result of receiving a platform session invite.
+	 * This event will trigger as a result of receiving a session invite.
 	 *
 	 * @return
 	 */
 	virtual TOnlineEvent<void(const FSessionInviteReceived&)> OnSessionInviteReceived() = 0;
 
 	/**
-	 * Get the event that is triggered when a session is joined via UI.
-	 * This event will trigger as a result of joining a session via the platform UI.
+	 * This event will trigger as a result of accepting a session invite or joining a session via the platform UI.
 	 *
 	 * @return
 	 */
@@ -1013,12 +1077,21 @@ END_ONLINE_STRUCT_META()
 BEGIN_ONLINE_STRUCT_META(FSendSessionInvite::Result)
 END_ONLINE_STRUCT_META()
 
-BEGIN_ONLINE_STRUCT_META(FGetSessionInvites::Params)
-	ONLINE_STRUCT_FIELD(FGetSessionInvites::Params, LocalAccountId)
+BEGIN_ONLINE_STRUCT_META(FGetSessionInviteById::Params)
+	ONLINE_STRUCT_FIELD(FGetSessionInviteById::Params, LocalAccountId),
+	ONLINE_STRUCT_FIELD(FGetSessionInviteById::Params, SessionInviteId)
+END_ONLINE_STRUCT_META()
+
+BEGIN_ONLINE_STRUCT_META(FGetSessionInviteById::Result)
+ONLINE_STRUCT_FIELD(FGetSessionInviteById::Result, SessionInvite)
+END_ONLINE_STRUCT_META()
+
+BEGIN_ONLINE_STRUCT_META(FGetAllSessionInvites::Params)
+	ONLINE_STRUCT_FIELD(FGetAllSessionInvites::Params, LocalAccountId)
 END_ONLINE_STRUCT_META()
 
-BEGIN_ONLINE_STRUCT_META(FGetSessionInvites::Result)
-	ONLINE_STRUCT_FIELD(FGetSessionInvites::Result, SessionInvites)
+BEGIN_ONLINE_STRUCT_META(FGetAllSessionInvites::Result)
+	ONLINE_STRUCT_FIELD(FGetAllSessionInvites::Result, SessionInvites)
 END_ONLINE_STRUCT_META()
 
 BEGIN_ONLINE_STRUCT_META(FRejectSessionInvite::Params)
diff --git a/Engine/Plugins/Online/Experimental/OnlineServicesEOSGS/Source/Private/Online/SessionsEOSGS.cpp b/Engine/Plugins/Online/Experimental/OnlineServicesEOSGS/Source/Private/Online/SessionsEOSGS.cpp
--- a/Engine/Plugins/Online/Experimental/OnlineServicesEOSGS/Source/Private/Online/SessionsEOSGS.cpp
+++ b/Engine/Plugins/Online/Experimental/OnlineServicesEOSGS/Source/Private/Online/SessionsEOSGS.cpp
@@ -241,7 +241,7 @@ void FSessionsEOSGS::HandleSessionInviteReceived(const EOS_Sessions_SessionInvit
 
 					AddSessionInvite(SessionInviteRef, Session, ReceiverId);
 
-					FSessionInviteReceived Event = { ReceiverId, SessionInviteRef };
+					FSessionInviteReceived Event = { ReceiverId, SessionInviteId };
 
 					SessionEvents.OnSessionInviteReceived.Broadcast(Event);
 				}
@@ -431,6 +431,11 @@ TOnlineAsyncOpHandle<FCreateSession> FSessionsEOSGS::CreateSession(FCreateSessio
 			NewSessionEOSGSRef->OwnerAccountId = OpParams.LocalAccountId;
 			NewSessionEOSGSRef->SessionSettings = OpParams.SessionSettings;
 			NewSessionEOSGSRef->SessionInfo.SessionId = CreateSessionId(Result.GetOkValue().NewSessionId);
+			NewSessionEOSGSRef->SessionInfo.bAllowSanctionedPlayers = OpParams.bAllowSanctionedPlayers;
+			NewSessionEOSGSRef->SessionInfo.bAntiCheatProtected = OpParams.bAntiCheatProtected;
+			NewSessionEOSGSRef->SessionInfo.bIsDedicatedServerSession = IsRunningDedicatedServer();
+			NewSessionEOSGSRef->SessionInfo.bIsLANSession = false;
+			NewSessionEOSGSRef->SessionInfo.SessionIdOverride = OpParams.SessionIdOverride;
 
 			AddSessionWithReferences(NewSessionEOSGSRef, OpParams.SessionName, OpParams.LocalAccountId, OpParams.bPresenceEnabled);
 
@@ -1215,8 +1220,6 @@ TOnlineAsyncOpHandle<FFindSessions> FSessionsEOSGS::FindSessions(FFindSessions::
 		});
 	});
 
-	// TODO: Call BuildSessionFromDetailsHandle as many times as we have details handles
-
 	Op->Enqueue(GetSerialQueue());
 
 	return Op->GetHandle();
diff --git a/Engine/Plugins/Online/Experimental/OnlineServicesEOSGS/Source/Public/Online/NboSerializerEOSGSSvc.h b/Engine/Plugins/Online/Experimental/OnlineServicesEOSGS/Source/Public/Online/NboSerializerEOSGSSvc.h
--- a/Engine/Plugins/Online/Experimental/OnlineServicesEOSGS/Source/Public/Online/NboSerializerEOSGSSvc.h
+++ b/Engine/Plugins/Online/Experimental/OnlineServicesEOSGS/Source/Public/Online/NboSerializerEOSGSSvc.h
@@ -33,7 +33,7 @@ inline void SerializeToBuffer(FNboSerializeToBuffer& Ar, const FOnlineSessionId&
 	Ar.WriteBinary(Data.GetData(), Data.Num());
 }
 
-inline void SerializeToBuffer(FNboSerializeToBuffer& Packet, const FSessionMembersSet& SessionMembersSet)
+inline void SerializeToBuffer(FNboSerializeToBuffer& Packet, const FSessionMemberIdsSet& SessionMembersSet)
 {
 	Packet << SessionMembersSet.Num();
 
@@ -63,7 +63,7 @@ inline void SerializeFromBuffer(FNboSerializeFromBuffer& Ar, FOnlineSessionId& S
 	SessionId = FOnlineSessionIdRegistryEOSGS::Get().FromReplicationData(Data);
 }
 
-inline void SerializeFromBuffer(FNboSerializeFromBuffer& Packet, FSessionMembersSet& SessionMembersSet)
+inline void SerializeFromBuffer(FNboSerializeFromBuffer& Packet, FSessionMemberIdsSet& SessionMembersSet)
 {
 	int32 NumEntries = 0;
 	Packet >> NumEntries;
diff --git a/Engine/Plugins/Online/Experimental/OnlineServicesNull/Source/Public/Online/NboSerializerNullSvc.h b/Engine/Plugins/Online/Experimental/OnlineServicesNull/Source/Public/Online/NboSerializerNullSvc.h
--- a/Engine/Plugins/Online/Experimental/OnlineServicesNull/Source/Public/Online/NboSerializerNullSvc.h
+++ b/Engine/Plugins/Online/Experimental/OnlineServicesNull/Source/Public/Online/NboSerializerNullSvc.h
@@ -33,7 +33,7 @@ inline void SerializeToBuffer(FNboSerializeToBuffer& Ar, const FOnlineSessionId&
 	Ar.WriteBinary(Data.GetData(), Data.Num());
 }
 
-inline void SerializeToBuffer(FNboSerializeToBuffer& Packet, const FSessionMembersSet& SessionMembersSet)
+inline void SerializeToBuffer(FNboSerializeToBuffer& Packet, const FSessionMemberIdsSet& SessionMembersSet)
 {
 	Packet << SessionMembersSet.Num();
 
@@ -75,7 +75,7 @@ inline void SerializeFromBuffer(FNboSerializeFromBuffer& Ar, FOnlineSessionId& S
 	SessionId = FOnlineSessionIdRegistryNull::Get().FromReplicationData(Data);
 }
 
-inline void SerializeFromBuffer(FNboSerializeFromBuffer& Packet, FSessionMembersSet& SessionMembersSet)
+inline void SerializeFromBuffer(FNboSerializeFromBuffer& Packet, FSessionMemberIdsSet& SessionMembersSet)
 {
 	int32 NumEntries = 0;
 	Packet >> NumEntries;
diff --git a/Engine/Plugins/Online/Experimental/OnlineServicesOSSAdapter/Source/Private/Online/SessionsOSSAdapter.cpp b/Engine/Plugins/Online/Experimental/OnlineServicesOSSAdapter/Source/Private/Online/SessionsOSSAdapter.cpp
--- a/Engine/Plugins/Online/Experimental/OnlineServicesOSSAdapter/Source/Private/Online/SessionsOSSAdapter.cpp
+++ b/Engine/Plugins/Online/Experimental/OnlineServicesOSSAdapter/Source/Private/Online/SessionsOSSAdapter.cpp
@@ -269,6 +269,12 @@ const FOnlineSession& FSessionOSSAdapter::GetV1Session() const
 	return V1Session;
 }
 
+void FSessionOSSAdapter::DumpState() const
+{
+	FSessionCommon::DumpState();
+	UE_LOG(LogTemp, Log, TEXT("V1Session: SessionId [%s]"), *V1Session.GetSessionIdStr());
+}
+
 /** FSessionsOSSAdapter */ 
 
 FSessionsOSSAdapter::FSessionsOSSAdapter(FOnlineServicesOSSAdapter& InServices)
@@ -302,7 +308,7 @@ void FSessionsOSSAdapter::Initialize()
 
 		AddSessionInvite(SessionInvite, Session, SessionInvite->RecipientId);
 
-		FSessionInviteReceived Event{ ServicesOSSAdapter.GetAccountIdRegistry().FindOrAddHandle(UserId.AsShared()), SessionInvite };
+		FSessionInviteReceived Event{ ServicesOSSAdapter.GetAccountIdRegistry().FindOrAddHandle(UserId.AsShared()), SessionInvite->InviteId };
 
 		SessionEvents.OnSessionInviteReceived.Broadcast(Event);
 	}));
diff --git a/Engine/Plugins/Online/Experimental/OnlineServicesOSSAdapter/Source/Public/Online/SessionsOSSAdapter.h b/Engine/Plugins/Online/Experimental/OnlineServicesOSSAdapter/Source/Public/Online/SessionsOSSAdapter.h
--- a/Engine/Plugins/Online/Experimental/OnlineServicesOSSAdapter/Source/Public/Online/SessionsOSSAdapter.h
+++ b/Engine/Plugins/Online/Experimental/OnlineServicesOSSAdapter/Source/Public/Online/SessionsOSSAdapter.h
@@ -44,6 +44,8 @@ public:
 
 	const FOnlineSession& GetV1Session() const;
 
+	virtual void DumpState() const override;
+
 private:
 	FOnlineSession V1Session;
 };
