Integrate schema into lobbies.

diff --git a/Engine/Config/BaseEngine.ini b/Engine/Config/BaseEngine.ini
--- a/Engine/Config/BaseEngine.ini
+++ b/Engine/Config/BaseEngine.ini
@@ -2158,6 +2158,14 @@ TravelSessionTimeoutSecs=45
 BeaconConnectionInitialTimeout=90.0
 BeaconConnectionTimeout=45.0
 
+[OnlineServices.Lobbies]
++SchemaDescriptors=(Id="LobbyBase", CategoryIds=("Lobby", "LobbyMember"))
++SchemaCategoryDescriptors=(Id="Lobby", ServiceDescriptorId="Lobby")
++SchemaCategoryDescriptors=(Id="LobbyMember", ServiceDescriptorId="LobbyMember")
++SchemaCategoryAttributeDescriptors=(SchemaId="LobbyBase", CategoryId="Lobby", AttributeIds=("SchemaCompatibilityId"))
++SchemaCategoryAttributeDescriptors=(SchemaId="LobbyBase", CategoryId="LobbyMember")
++SchemaAttributeDescriptors=(Id="SchemaCompatibilityId", Type="Int64", Flags=("Public", "SchemaCompatibilityId"))
+
 [StaticMeshLODSettings]
 LevelArchitecture=(NumLODs=4,MaxNumStreamedLODs=0,bSupportLODStreaming=0,LightMapResolution=32,LODPercentTriangles=50,PixelError=12,SilhouetteImportance=4,Name=LOCTEXT("LevelArchitectureLOD","Level Architecture"))
 SmallProp=(NumLODs=4,MaxNumStreamedLODs=0,bSupportLODStreaming=0,LODPercentTriangles=50,PixelError=10,Name=LOCTEXT("SmallPropLOD","Small Prop"))
diff --git a/Engine/Plugins/Online/Experimental/OnlineServices/Source/OnlineServicesCommon/Private/Online/LobbiesCommon.cpp b/Engine/Plugins/Online/Experimental/OnlineServices/Source/OnlineServicesCommon/Private/Online/LobbiesCommon.cpp
--- a/Engine/Plugins/Online/Experimental/OnlineServices/Source/OnlineServicesCommon/Private/Online/LobbiesCommon.cpp
+++ b/Engine/Plugins/Online/Experimental/OnlineServices/Source/OnlineServicesCommon/Private/Online/LobbiesCommon.cpp
@@ -10,6 +10,7 @@ namespace UE::Online {
 
 FLobbiesCommon::FLobbiesCommon(FOnlineServicesCommon& InServices)
 	: TOnlineComponent(TEXT("Lobbies"), InServices)
+	, SchemaRegistry(MakeShared<FSchemaRegistry>())
 {
 }
 
@@ -17,32 +18,12 @@ void FLobbiesCommon::Initialize()
 {
 	TOnlineComponent<ILobbies>::Initialize();
 
-	LobbySchemaRegistry = MakeShared<FLobbySchemaRegistry>();
-	ServiceSchema = MakeShared<FLobbySchema>();
+	FSchemaRegistryDescriptorConfig SchemaConfig;
+	LoadConfig(SchemaConfig);
 
-	FLobbyConfig LobbyConfig;
-	// Todo: LoadConfig(LobbyConfig);
-
-	if (LobbyConfig.RegisteredSchemas.Num() > 0)
-	{
-		TArray<FLobbySchemaConfig> LobbySchemaConfigs;
-		LobbySchemaConfigs.Reserve(LobbyConfig.RegisteredSchemas.Num());
-
-		for (FLobbySchemaId LobbySchemaId : LobbyConfig.RegisteredSchemas)
-		{
-			FLobbySchemaConfig LobbySchemaConfig;
-			// Todo: LoadConfig(LobbySchemaConfig, LobbySchemaId.ToString());
-			LobbySchemaConfigs.Add(MoveTemp(LobbySchemaConfig));
-		}
-
-		if (!LobbySchemaRegistry->Initialize(MoveTemp(LobbySchemaConfigs)))
-		{
-			// Todo: Log error.
-		}
-	}
-	else
+	if (!SchemaRegistry->ParseConfig(SchemaConfig))
 	{
-		// Todo: Log error.
+		UE_LOG(LogOnlineServices, Error, TEXT("[FLobbiesCommon::Initialize] Failed to initialize schema registry"));
 	}
 }
 
