Integrate schema into lobbies.

diff --git a/Engine/Plugins/Online/Experimental/OnlineServices/Source/OnlineServicesCommon/Private/Online/LobbiesCommonTests.cpp b/Engine/Plugins/Online/Experimental/OnlineServices/Source/OnlineServicesCommon/Private/Online/LobbiesCommonTests.cpp
--- a/Engine/Plugins/Online/Experimental/OnlineServices/Source/OnlineServicesCommon/Private/Online/LobbiesCommonTests.cpp
+++ b/Engine/Plugins/Online/Experimental/OnlineServices/Source/OnlineServicesCommon/Private/Online/LobbiesCommonTests.cpp
@@ -2,6 +2,7 @@
 
 #include "LobbiesCommonTests.h"
 
+#include "Algo/Transform.h"
 #include "Containers/Ticker.h"
 #include "Misc/AutomationTest.h"
 #include "Online/LobbiesCommon.h"
@@ -110,157 +111,371 @@ IMPLEMENT_SIMPLE_AUTOMATION_TEST(FClientLobbyDataTest,
 	EAutomationTestFlags::EditorContext | EAutomationTestFlags::EngineFilter)
 bool FClientLobbyDataTest::RunTest(const FString& Parameters)
 {
+	FSchemaId LobbySchemaId1 = TEXT("Lobby1");
+	FSchemaId LobbySchemaId2 = TEXT("Lobby2");
+	FSchemaCategoryId LobbyCategoryId = TEXT("Lobby");
+	FSchemaCategoryId LobbyMemberCategoryId = TEXT("LobbyMember");
+	FSchemaServiceDescriptorId LobbyServiceDescriptorId = TEXT("LobbyServiceDescriptor");
+	FSchemaServiceAttributeId LobbyServiceAttributeId1 = TEXT("LobbyServiceAttributeId1");
+	FSchemaServiceAttributeId LobbyServiceAttributeId2 = TEXT("LobbyServiceAttributeId2");
+	FSchemaServiceAttributeId LobbyServiceAttributeId3 = TEXT("LobbyServiceAttributeId3");
+	FSchemaServiceDescriptorId LobbyMemberServiceDescriptorId = TEXT("LobbyMemberServiceDescriptor");
+	FSchemaServiceAttributeId LobbyMemberServiceAttributeId1 = TEXT("LobbyMemberServiceAttributeId1");
+	FSchemaServiceAttributeId LobbyMemberServiceAttributeId2 = TEXT("LobbyMemberServiceAttributeId2");
+	const TArray<ESchemaServiceAttributeSupportedTypeFlags> AllSupportedServiceAttributeTypes = {
+		ESchemaServiceAttributeSupportedTypeFlags::Bool,
+		ESchemaServiceAttributeSupportedTypeFlags::Int64,
+		ESchemaServiceAttributeSupportedTypeFlags::Double,
+		ESchemaServiceAttributeSupportedTypeFlags::String };
+	const TArray<ESchemaServiceAttributeFlags> AllSupportedServiceAttributeFlags = {
+		ESchemaServiceAttributeFlags::Searchable,
+		ESchemaServiceAttributeFlags::Public,
+		ESchemaServiceAttributeFlags::Private };
+
+	FSchemaAttributeId SchemaCompatibilityAttributeId = TEXT("SchemaCompatibilityAttribute");
+	FSchemaAttributeId LobbyAttributeId1 = TEXT("LobbyAttribute1");
+	FSchemaAttributeId LobbyAttributeId2 = TEXT("LobbyAttribute2");
+	FSchemaAttributeId LobbyMemberAttributeId1 = TEXT("LobbyMemberAttribute1");
+	FSchemaAttributeId LobbyMemberAttributeId2 = TEXT("LobbyMemberAttribute2");
+
+	FSchemaRegistryDescriptorConfig SchemaConfig;
+	SchemaConfig.SchemaDescriptors.Add({ LobbyBaseSchemaId, FSchemaId(), { LobbyCategoryId, LobbyMemberCategoryId } });
+	SchemaConfig.SchemaCategoryDescriptors.Add({ LobbyCategoryId, LobbyServiceDescriptorId });
+	SchemaConfig.SchemaCategoryDescriptors.Add({ LobbyMemberCategoryId, LobbyMemberServiceDescriptorId });
+	SchemaConfig.SchemaCategoryAttributeDescriptors.Add({ LobbyBaseSchemaId, LobbyCategoryId, { SchemaCompatibilityAttributeId } });
+	SchemaConfig.SchemaCategoryAttributeDescriptors.Add({ LobbyBaseSchemaId, LobbyMemberCategoryId, {} });
+	SchemaConfig.SchemaAttributeDescriptors.Add({ SchemaCompatibilityAttributeId, ESchemaAttributeType::Int64, { ESchemaAttributeFlags::Public, ESchemaAttributeFlags::SchemaCompatibilityId }, 0, 0 });
+
+	SchemaConfig.SchemaDescriptors.Add({ LobbySchemaId1, LobbyBaseSchemaId });
+	SchemaConfig.SchemaCategoryAttributeDescriptors.Add({ LobbySchemaId1, LobbyCategoryId, { LobbyAttributeId1, LobbyAttributeId2 } });
+	SchemaConfig.SchemaCategoryAttributeDescriptors.Add({ LobbySchemaId1, LobbyMemberCategoryId, { LobbyMemberAttributeId1, LobbyMemberAttributeId2 } });
+
+	SchemaConfig.SchemaDescriptors.Add({ LobbySchemaId2, LobbyBaseSchemaId });
+	SchemaConfig.SchemaCategoryAttributeDescriptors.Add({ LobbySchemaId2, LobbyCategoryId, { LobbyAttributeId1, LobbyAttributeId2 } });
+	SchemaConfig.SchemaCategoryAttributeDescriptors.Add({ LobbySchemaId2, LobbyMemberCategoryId, { LobbyMemberAttributeId1, LobbyMemberAttributeId2 } });
+
+	SchemaConfig.SchemaAttributeDescriptors.Add({ LobbyAttributeId1, ESchemaAttributeType::String, { ESchemaAttributeFlags::Public }, 0, 64 });
+	SchemaConfig.SchemaAttributeDescriptors.Add({ LobbyAttributeId2, ESchemaAttributeType::String, { ESchemaAttributeFlags::Public }, 0, 64 });
+	SchemaConfig.SchemaAttributeDescriptors.Add({ LobbyMemberAttributeId1, ESchemaAttributeType::String, { ESchemaAttributeFlags::Public }, 0, 64 });
+	SchemaConfig.SchemaAttributeDescriptors.Add({ LobbyMemberAttributeId2, ESchemaAttributeType::String, { ESchemaAttributeFlags::Public }, 0, 64 });
+
+	SchemaConfig.ServiceDescriptors.Add({ LobbyServiceDescriptorId, { LobbyServiceAttributeId1, LobbyServiceAttributeId2, LobbyServiceAttributeId3 } });
+	SchemaConfig.ServiceAttributeDescriptors.Add({ LobbyServiceAttributeId1, AllSupportedServiceAttributeTypes, AllSupportedServiceAttributeFlags, 64 });
+	SchemaConfig.ServiceAttributeDescriptors.Add({ LobbyServiceAttributeId2, AllSupportedServiceAttributeTypes, AllSupportedServiceAttributeFlags, 64 });
+	SchemaConfig.ServiceAttributeDescriptors.Add({ LobbyServiceAttributeId3, AllSupportedServiceAttributeTypes, AllSupportedServiceAttributeFlags, 64 });
+	SchemaConfig.ServiceDescriptors.Add({ LobbyMemberServiceDescriptorId, { LobbyMemberServiceAttributeId1, LobbyMemberServiceAttributeId2 } });
+	SchemaConfig.ServiceAttributeDescriptors.Add({ LobbyMemberServiceAttributeId1, AllSupportedServiceAttributeTypes, AllSupportedServiceAttributeFlags, 64 });
+	SchemaConfig.ServiceAttributeDescriptors.Add({ LobbyMemberServiceAttributeId2, AllSupportedServiceAttributeTypes, AllSupportedServiceAttributeFlags, 64 });
+
+	TSharedRef<FSchemaRegistry> SchemaRegistry = MakeShared<FSchemaRegistry>();
+	UTEST_TRUE("Schema config parsed successfully.", SchemaRegistry->ParseConfig(SchemaConfig));
+
+	TSharedPtr<const FSchemaDefinition> LobbySchemaDefinition = SchemaRegistry->GetDefinition(LobbySchemaId1);
+	UTEST_NOT_NULL("Lobby schema definition exists.", LobbySchemaDefinition.Get());
+
+	FSchemaVariant SchemaCompatibilityAttributeValue(LobbySchemaDefinition->CompatibilityId);
+	FSchemaVariant LobbyAttribute1Value1(TEXT("Attribute1-Value1"));
+	FSchemaVariant LobbyAttribute1Value2(TEXT("Attribute1-Value2"));
+	FSchemaVariant LobbyAttribute2Value1(TEXT("Attribute2-Value1"));
+	FSchemaVariant LobbyAttribute2Value2(TEXT("Attribute2-Value2"));
+	FSchemaVariant LobbyMemberAttribute1Value1(TEXT("MemberAttribute1-Value1"));
+	FSchemaVariant LobbyMemberAttribute1Value2(TEXT("MemberAttribute1-Value2"));
+	FSchemaVariant LobbyMemberAttribute2Value1(TEXT("MemberAttribute2-Value1"));
+	FSchemaVariant LobbyMemberAttribute2Value2(TEXT("MemberAttribute2-Value2"));
+
 	// 1. Applying snapshot with no local members does not fire events.
 	{
-		FLobbyEvents LobbyEvents;
-		FLobbyEventCapture EventCapture(LobbyEvents);
-		FLobbyId LobbyId(EOnlineServices::Null, 1);
-		FClientLobbyData ClientData(LobbyId);
-
 		FAccountId User1(EOnlineServices::Null, 1);
 		FAccountId User2(EOnlineServices::Null, 2);
-		FName LobbySchemaName = TEXT("SchemaName");
 		int32 LobbyMaxMembers = 5;
 		ELobbyJoinPolicy LobbyJoinPolicy = ELobbyJoinPolicy::PublicAdvertised;
-		FLobbyAttributeId LobbyAttribute1Key = TEXT("LobbyAttribute1");
-		FLobbyVariant LobbyAttribute1Value;
-		LobbyAttribute1Value.Set(TEXT("Attribute1"));
-		FLobbyAttributeId LobbyAttribute2Key = TEXT("LobbyAttribute2");
-		FLobbyVariant LobbyAttribute2Value;
-		LobbyAttribute2Value.Set(TEXT("Attribute2"));
-		FLobbyAttributeId LobbyMemberAttribute1Key = TEXT("LobbyMemberAttribute1");
-		FLobbyVariant LobbyMemberAttribute1Value;
-		LobbyMemberAttribute1Value.Set(TEXT("MemberAttribute1"));
-		FLobbyAttributeId LobbyMemberAttribute2Key = TEXT("LobbyMemberAttribute2");
-		FLobbyVariant LobbyMemberAttribute2Value;
-		LobbyMemberAttribute2Value.Set(TEXT("MemberAttribute2"));
-
-		// 1.1 setup.
-		{
-			TMap<FAccountId, TSharedRef<FClientLobbyMemberSnapshot>> LobbyMemberSnapshots;
+
+		FLobbyEvents LobbyEvents;
+		FLobbyEventCapture EventCapture(LobbyEvents);
+		FLobbyId LobbyId(EOnlineServices::Null, 1);
+		FLobbyClientData ClientData(LobbyId, SchemaRegistry);
+
+		// 1.1 setup. Create lobby data from snapshot with initial attributes.
+		{
+			TMap<FAccountId, FLobbyMemberServiceSnapshot> LobbyMemberSnapshots;
 			{
-				TSharedRef<FClientLobbyMemberSnapshot> MemberSnapshot = MakeShared<FClientLobbyMemberSnapshot>();
-				MemberSnapshot->AccountId = User1;
-				MemberSnapshot->Attributes.Add(LobbyMemberAttribute1Key, LobbyMemberAttribute1Value);
-				LobbyMemberSnapshots.Add(User1, MemberSnapshot);
+				FLobbyMemberServiceSnapshot& MemberSnapshot = LobbyMemberSnapshots.Add(User1);
+				MemberSnapshot.AccountId = User1;
+				MemberSnapshot.SchemaServiceSnapshot.Attributes.Add(LobbyMemberServiceAttributeId1, LobbyMemberAttribute1Value1);
 			}
 			{
-				TSharedRef<FClientLobbyMemberSnapshot> MemberSnapshot = MakeShared<FClientLobbyMemberSnapshot>();
-				MemberSnapshot->AccountId = User2;
-				LobbyMemberSnapshots.Add(User2, MemberSnapshot);
+				FLobbyMemberServiceSnapshot& MemberSnapshot = LobbyMemberSnapshots.Add(User2);
+				MemberSnapshot.AccountId = User2;
 			}
 
-			FClientLobbySnapshot LobbySnapshot;
+			FLobbyServiceSnapshot LobbySnapshot;
 			LobbySnapshot.OwnerAccountId = User1;
-			LobbySnapshot.SchemaName = LobbySchemaName;
 			LobbySnapshot.MaxMembers = LobbyMaxMembers;
 			LobbySnapshot.JoinPolicy = LobbyJoinPolicy;
-			LobbySnapshot.Attributes.Add(LobbyAttribute1Key, LobbyAttribute1Value);
+			LobbySnapshot.SchemaServiceSnapshot.Attributes.Add(LobbyServiceAttributeId1, SchemaCompatibilityAttributeValue);
+			LobbySnapshot.SchemaServiceSnapshot.Attributes.Add(LobbyServiceAttributeId2, LobbyAttribute1Value1);
 			LobbySnapshot.Members = {User1, User2};
 
 			TMap<FAccountId, ELobbyMemberLeaveReason> LeaveReasons;
-			FApplyLobbyUpdateResult Result = ClientData.ApplyLobbyUpdateFromServiceSnapshot(MoveTemp(LobbySnapshot), MoveTemp(LobbyMemberSnapshots), MoveTemp(LeaveReasons), &LobbyEvents);
-			UTEST_EQUAL("Check no local members left.", Result.LeavingLocalMembers.Num(), 0);
-		}
-
-		// 1.1.1. Initial snapshot does not trigger any events.
-		UTEST_EQUAL("Snapshot does not trigger any events.", EventCapture.GetTotalNotificationsReceived(), 0);
-		UTEST_EQUAL("Verify lobby owner.", ClientData.GetPublicData().OwnerAccountId, User1);
-		UTEST_EQUAL("Verify lobby schema.", ClientData.GetPublicData().SchemaName, LobbySchemaName);
-		UTEST_EQUAL("Verify lobby max members.", ClientData.GetPublicData().MaxMembers, LobbyMaxMembers);
-		UTEST_EQUAL("Verify lobby join policy.", ClientData.GetPublicData().JoinPolicy, LobbyJoinPolicy);
-		UTEST_EQUAL("Verify lobby attributes.", ClientData.GetPublicData().Attributes.Num(), 1);
-		UTEST_NOT_NULL("Verify lobby attributes.", ClientData.GetPublicData().Attributes.Find(LobbyAttribute1Key));
-		UTEST_EQUAL("Verify lobby attributes.", ClientData.GetPublicData().Attributes[LobbyAttribute1Key], LobbyAttribute1Value);
-		UTEST_EQUAL("Verify lobby members.", ClientData.GetPublicData().Members.Num(), 2);
-		UTEST_NOT_NULL("Verify lobby members.", ClientData.GetPublicData().Members.Find(User1));
-		UTEST_EQUAL("Verify lobby members.", ClientData.GetPublicData().Members[User1]->AccountId, User1);
-		UTEST_NOT_NULL("Verify lobby members.", ClientData.GetPublicData().Members.Find(User2));
-		UTEST_EQUAL("Verify lobby members.", ClientData.GetPublicData().Members[User2]->AccountId, User2);
-		UTEST_EQUAL("Verify lobby member attributes.", ClientData.GetPublicData().Members[User1]->Attributes.Num(), 1);
-		UTEST_NOT_NULL("Verify lobby member attributes.", ClientData.GetPublicData().Members[User1]->Attributes.Find(LobbyMemberAttribute1Key));
-		UTEST_EQUAL("Verify lobby member attributes.", ClientData.GetPublicData().Members[User1]->Attributes[LobbyMemberAttribute1Key], LobbyMemberAttribute1Value);
-
-		// 1.2 setup.
-		{
-			LobbyAttribute1Value.Set(TEXT("ModifiedAttribute1"));
-			LobbyMemberAttribute1Value.Set(TEXT("ModifiedMemberAttribute1"));
-			LobbySchemaName = TEXT("ModifiedSchemaName");
-
-			TMap<FAccountId, TSharedRef<FClientLobbyMemberSnapshot>> LobbyMemberSnapshots;
+			TOnlineResult<FLobbyClientDataPrepareServiceSnapshot> PrepareResult = ClientData.PrepareServiceSnapshot({ MoveTemp(LobbySnapshot), MoveTemp(LobbyMemberSnapshots), MoveTemp(LeaveReasons)});
+			UTEST_TRUE("Check PrepareServiceSnapshot succeeded.", PrepareResult.IsOk());
+
+			FLobbyClientDataCommitServiceSnapshot::Result CommitResult = ClientData.CommitServiceSnapshot({ &LobbyEvents });
+			UTEST_EQUAL("Check no local members left.", CommitResult.LeavingLocalMembers.Num(), 0);
+		}
+
+		{
+			// 1.1.1. Initial snapshot does not trigger any events.
+			UTEST_EQUAL("Snapshot does not trigger any events.", EventCapture.GetTotalNotificationsReceived(), 0);
+			UTEST_EQUAL("Verify lobby owner.", ClientData.GetPublicData().OwnerAccountId, User1);
+			UTEST_EQUAL("Verify lobby schema.", ClientData.GetPublicData().SchemaId, LobbySchemaId1);
+			UTEST_EQUAL("Verify lobby max members.", ClientData.GetPublicData().MaxMembers, LobbyMaxMembers);
+			UTEST_EQUAL("Verify lobby join policy.", ClientData.GetPublicData().JoinPolicy, LobbyJoinPolicy);
+			UTEST_EQUAL("Verify lobby attributes.", ClientData.GetPublicData().Attributes.Num(), 1);
+			UTEST_NOT_NULL("Verify lobby attributes.", ClientData.GetPublicData().Attributes.Find(LobbyAttributeId1));
+			UTEST_EQUAL("Verify lobby attributes.", ClientData.GetPublicData().Attributes[LobbyAttributeId1], LobbyAttribute1Value1);
+			UTEST_EQUAL("Verify lobby members.", ClientData.GetPublicData().Members.Num(), 2);
+			UTEST_NOT_NULL("Verify lobby members.", ClientData.GetPublicData().Members.Find(User1));
+			UTEST_EQUAL("Verify lobby members.", ClientData.GetPublicData().Members[User1]->AccountId, User1);
+			UTEST_NOT_NULL("Verify lobby members.", ClientData.GetPublicData().Members.Find(User2));
+			UTEST_EQUAL("Verify lobby members.", ClientData.GetPublicData().Members[User2]->AccountId, User2);
+			UTEST_EQUAL("Verify lobby member attributes.", ClientData.GetPublicData().Members[User1]->Attributes.Num(), 1);
+			UTEST_NOT_NULL("Verify lobby member attributes.", ClientData.GetPublicData().Members[User1]->Attributes.Find(LobbyMemberAttributeId1));
+			UTEST_EQUAL("Verify lobby member attributes.", ClientData.GetPublicData().Members[User1]->Attributes[LobbyMemberAttributeId1], LobbyMemberAttribute1Value1);
+		}
+
+		// 1.2 setup. Change owner. No member snapshots due to no detected member changes.
+		{
+			TMap<FAccountId, FLobbyMemberServiceSnapshot> LobbyMemberSnapshots;
+			FLobbyServiceSnapshot LobbySnapshot;
+			LobbySnapshot.OwnerAccountId = User2;
+			LobbySnapshot.MaxMembers = LobbyMaxMembers;
+			LobbySnapshot.JoinPolicy = LobbyJoinPolicy;
+			LobbySnapshot.SchemaServiceSnapshot.Attributes.Add(LobbyServiceAttributeId1, SchemaCompatibilityAttributeValue);
+			LobbySnapshot.SchemaServiceSnapshot.Attributes.Add(LobbyServiceAttributeId2, LobbyAttribute1Value1);
+			LobbySnapshot.Members = {User1, User2};
+
+			TMap<FAccountId, ELobbyMemberLeaveReason> LeaveReasons;
+			TOnlineResult<FLobbyClientDataPrepareServiceSnapshot> PrepareResult = ClientData.PrepareServiceSnapshot({ MoveTemp(LobbySnapshot), MoveTemp(LobbyMemberSnapshots), MoveTemp(LeaveReasons) });
+			UTEST_TRUE("Check PrepareServiceSnapshot succeeded.", PrepareResult.IsOk());
+
+			FLobbyClientDataCommitServiceSnapshot::Result CommitResult = ClientData.CommitServiceSnapshot({ &LobbyEvents });
+			UTEST_EQUAL("Check no local members left.", CommitResult.LeavingLocalMembers.Num(), 0);
+		}
+
+		{
+			// 1.2.1. Snapshot with changed lobby leader does not trigger lobby leader changed.
+			UTEST_EQUAL("Snapshot does not trigger any events.", EventCapture.GetTotalNotificationsReceived(), 0);
+			UTEST_EQUAL("Verify lobby owner.", ClientData.GetPublicData().OwnerAccountId, User2);
+			UTEST_EQUAL("Verify lobby attributes.", ClientData.GetPublicData().Attributes.Num(), 1);
+			const FSchemaVariant* LobbyAttribute1 = ClientData.GetPublicData().Attributes.Find(LobbyAttributeId1);
+			UTEST_NOT_NULL("Verify lobby attributes.", LobbyAttribute1);
+			UTEST_EQUAL("Verify lobby attributes.", *LobbyAttribute1, LobbyAttribute1Value1);
+			UTEST_EQUAL("Verify lobby members.", ClientData.GetPublicData().Members.Num(), 2);
+			const TSharedRef<const FLobbyMember>* User1Data = ClientData.GetPublicData().Members.Find(User1);
+			UTEST_NOT_NULL("Verify lobby members.", User1Data);
+			UTEST_NOT_NULL("Verify lobby members.", ClientData.GetPublicData().Members.Find(User2));
+			UTEST_EQUAL("Verify lobby member attributes.", (*User1Data)->Attributes.Num(), 1);
+			const FSchemaVariant* LobbyMemberAttribute1 = (*User1Data)->Attributes.Find(LobbyMemberAttributeId1);
+			UTEST_NOT_NULL("Verify lobby member attributes.", LobbyMemberAttribute1);
+			UTEST_EQUAL("Verify lobby member attributes.", *LobbyMemberAttribute1, LobbyMemberAttribute1Value1);
+		}
+
+		// 1.3 setup. Add lobby and lobby member attribute.
+		{
+			TMap<FAccountId, FLobbyMemberServiceSnapshot> LobbyMemberSnapshots;
 			{
-				TSharedRef<FClientLobbyMemberSnapshot> MemberSnapshot = MakeShared<FClientLobbyMemberSnapshot>();
-				MemberSnapshot->AccountId = User1;
-				MemberSnapshot->Attributes.Add(LobbyMemberAttribute1Key, LobbyMemberAttribute1Value);
-				LobbyMemberSnapshots.Add(User1, MemberSnapshot);
+				FLobbyMemberServiceSnapshot& MemberSnapshot = LobbyMemberSnapshots.Add(User1);
+				MemberSnapshot.AccountId = User1;
+				MemberSnapshot.SchemaServiceSnapshot.Attributes.Add(LobbyMemberServiceAttributeId1, LobbyMemberAttribute1Value1);
+				MemberSnapshot.SchemaServiceSnapshot.Attributes.Add(LobbyMemberServiceAttributeId2, LobbyMemberAttribute2Value1);
 			}
 			{
-				TSharedRef<FClientLobbyMemberSnapshot> MemberSnapshot = MakeShared<FClientLobbyMemberSnapshot>();
-				MemberSnapshot->AccountId = User2;
-				LobbyMemberSnapshots.Add(User2, MemberSnapshot);
+				FLobbyMemberServiceSnapshot& MemberSnapshot = LobbyMemberSnapshots.Add(User2);
+				MemberSnapshot.AccountId = User2;
 			}
 
-			FClientLobbySnapshot LobbySnapshot;
+			FLobbyServiceSnapshot LobbySnapshot;
 			LobbySnapshot.OwnerAccountId = User2;
-			LobbySnapshot.SchemaName = LobbySchemaName;
 			LobbySnapshot.MaxMembers = LobbyMaxMembers;
 			LobbySnapshot.JoinPolicy = LobbyJoinPolicy;
-			LobbySnapshot.Attributes.Add(LobbyAttribute1Key, LobbyAttribute1Value);
-			LobbySnapshot.Members = {User1, User2};
+			LobbySnapshot.SchemaServiceSnapshot.Attributes.Add(LobbyServiceAttributeId1, SchemaCompatibilityAttributeValue);
+			LobbySnapshot.SchemaServiceSnapshot.Attributes.Add(LobbyServiceAttributeId2, LobbyAttribute1Value1);
+			LobbySnapshot.SchemaServiceSnapshot.Attributes.Add(LobbyServiceAttributeId3, LobbyAttribute2Value1);
+			LobbySnapshot.Members = { User1, User2 };
 
 			TMap<FAccountId, ELobbyMemberLeaveReason> LeaveReasons;
-			FApplyLobbyUpdateResult Result = ClientData.ApplyLobbyUpdateFromServiceSnapshot(MoveTemp(LobbySnapshot), MoveTemp(LobbyMemberSnapshots), MoveTemp(LeaveReasons), &LobbyEvents);
-			UTEST_EQUAL("Check no local members left.", Result.LeavingLocalMembers.Num(), 0);
-		}
-
-		// 1.2.1. Snapshot with changed lobby attribute does not trigger lobby attribute changed.
-		// 1.2.2. Snapshot with changed lobby member attribute does not trigger lobby member attribute changed.
-		// 1.2.3. Snapshot with changed lobby schema does not trigger lobby schema changed.
-		// 1.2.4. Snapshot with changed lobby leader does not trigger lobby leader changed.
-		UTEST_EQUAL("Snapshot does not trigger any events.", EventCapture.GetTotalNotificationsReceived(), 0);
-		UTEST_EQUAL("Verify lobby owner.", ClientData.GetPublicData().OwnerAccountId, User2);
-		UTEST_EQUAL("Verify lobby schema.", ClientData.GetPublicData().SchemaName, LobbySchemaName);
-		UTEST_EQUAL("Verify lobby max members.", ClientData.GetPublicData().MaxMembers, LobbyMaxMembers);
-		UTEST_EQUAL("Verify lobby join policy.", ClientData.GetPublicData().JoinPolicy, LobbyJoinPolicy);
-		UTEST_EQUAL("Verify lobby attributes.", ClientData.GetPublicData().Attributes.Num(), 1);
-		UTEST_NOT_NULL("Verify lobby attributes.", ClientData.GetPublicData().Attributes.Find(LobbyAttribute1Key));
-		UTEST_EQUAL("Verify lobby attributes.", ClientData.GetPublicData().Attributes[LobbyAttribute1Key], LobbyAttribute1Value);
-		UTEST_EQUAL("Verify lobby members.", ClientData.GetPublicData().Members.Num(), 2);
-		UTEST_NOT_NULL("Verify lobby members.", ClientData.GetPublicData().Members.Find(User1));
-		UTEST_NOT_NULL("Verify lobby members.", ClientData.GetPublicData().Members.Find(User2));
-		UTEST_EQUAL("Verify lobby member attributes.", ClientData.GetPublicData().Members[User1]->Attributes.Num(), 1);
-		UTEST_NOT_NULL("Verify lobby member attributes.", ClientData.GetPublicData().Members[User1]->Attributes.Find(LobbyMemberAttribute1Key));
-		UTEST_EQUAL("Verify lobby member attributes.", ClientData.GetPublicData().Members[User1]->Attributes[LobbyMemberAttribute1Key], LobbyMemberAttribute1Value);
-
-		// 1.3 setup.
-		{
-			TMap<FAccountId, TSharedRef<FClientLobbyMemberSnapshot>> LobbyMemberSnapshots;
+			TOnlineResult<FLobbyClientDataPrepareServiceSnapshot> PrepareResult = ClientData.PrepareServiceSnapshot({ MoveTemp(LobbySnapshot), MoveTemp(LobbyMemberSnapshots), MoveTemp(LeaveReasons) });
+			UTEST_TRUE("Check PrepareServiceSnapshot succeeded.", PrepareResult.IsOk());
+
+			FLobbyClientDataCommitServiceSnapshot::Result CommitResult = ClientData.CommitServiceSnapshot({ &LobbyEvents });
+			UTEST_EQUAL("Check no local members left.", CommitResult.LeavingLocalMembers.Num(), 0);
+		}
+
+		{
+			// 1.3.1. Snapshot with changed lobby attribute does not trigger lobby attribute changed.
+			// 1.3.2. Snapshot with changed lobby member attribute does not trigger lobby member attribute changed.
+			UTEST_EQUAL("Snapshot does not trigger any events.", EventCapture.GetTotalNotificationsReceived(), 0);
+			UTEST_EQUAL("Verify lobby owner.", ClientData.GetPublicData().OwnerAccountId, User2);
+			UTEST_EQUAL("Verify lobby attributes.", ClientData.GetPublicData().Attributes.Num(), 2);
+			const FSchemaVariant* LobbyAttribute1 = ClientData.GetPublicData().Attributes.Find(LobbyAttributeId1);
+			UTEST_NOT_NULL("Verify lobby attributes.", LobbyAttribute1);
+			UTEST_EQUAL("Verify lobby attributes.", *LobbyAttribute1, LobbyAttribute1Value1);
+			const FSchemaVariant* LobbyAttribute2 = ClientData.GetPublicData().Attributes.Find(LobbyAttributeId2);
+			UTEST_NOT_NULL("Verify lobby attributes.", LobbyAttribute2);
+			UTEST_EQUAL("Verify lobby attributes.", *LobbyAttribute2, LobbyAttribute2Value1);
+			UTEST_EQUAL("Verify lobby members.", ClientData.GetPublicData().Members.Num(), 2);
+			const TSharedRef<const FLobbyMember>* User1Data = ClientData.GetPublicData().Members.Find(User1);
+			UTEST_NOT_NULL("Verify lobby members.", User1Data);
+			UTEST_EQUAL("Verify lobby member attributes.", (*User1Data)->Attributes.Num(), 2);
+			const FSchemaVariant* LobbyMemberAttribute1 = (*User1Data)->Attributes.Find(LobbyMemberAttributeId1);
+			UTEST_NOT_NULL("Verify lobby member attributes.", LobbyMemberAttribute1);
+			UTEST_EQUAL("Verify lobby member attributes.", *LobbyMemberAttribute1, LobbyMemberAttribute1Value1);
+			const FSchemaVariant* LobbyMemberAttribute2 = (*User1Data)->Attributes.Find(LobbyMemberAttributeId2);
+			UTEST_NOT_NULL("Verify lobby member attributes.", LobbyMemberAttribute2);
+			UTEST_EQUAL("Verify lobby member attributes.", *LobbyMemberAttribute2, LobbyMemberAttribute2Value1);
+			const TSharedRef<const FLobbyMember>* User2Data = ClientData.GetPublicData().Members.Find(User2);
+			UTEST_NOT_NULL("Verify lobby members.", User2Data);
+			UTEST_EQUAL("Verify lobby member attributes.", (*User2Data)->Attributes.Num(), 0);
+		}
+
+		// 1.3 setup. Modify lobby and lobby member attribute.
+		{
+			TMap<FAccountId, FLobbyMemberServiceSnapshot> LobbyMemberSnapshots;
 			{
-				TSharedRef<FClientLobbyMemberSnapshot> MemberSnapshot = MakeShared<FClientLobbyMemberSnapshot>();
-				MemberSnapshot->AccountId = User1;
-				LobbyMemberSnapshots.Add(User1, MemberSnapshot);
+				FLobbyMemberServiceSnapshot& MemberSnapshot = LobbyMemberSnapshots.Add(User1);
+				MemberSnapshot.AccountId = User1;
+				MemberSnapshot.SchemaServiceSnapshot.Attributes.Add(LobbyMemberServiceAttributeId1, LobbyMemberAttribute1Value1);
+				MemberSnapshot.SchemaServiceSnapshot.Attributes.Add(LobbyMemberServiceAttributeId2, LobbyMemberAttribute2Value2);
 			}
+			{
+				FLobbyMemberServiceSnapshot& MemberSnapshot = LobbyMemberSnapshots.Add(User2);
+				MemberSnapshot.AccountId = User2;
+			}
+
+			FLobbyServiceSnapshot LobbySnapshot;
+			LobbySnapshot.OwnerAccountId = User2;
+			LobbySnapshot.MaxMembers = LobbyMaxMembers;
+			LobbySnapshot.JoinPolicy = LobbyJoinPolicy;
+			LobbySnapshot.SchemaServiceSnapshot.Attributes.Add(LobbyServiceAttributeId1, SchemaCompatibilityAttributeValue);
+			LobbySnapshot.SchemaServiceSnapshot.Attributes.Add(LobbyServiceAttributeId2, LobbyAttribute1Value1);
+			LobbySnapshot.SchemaServiceSnapshot.Attributes.Add(LobbyServiceAttributeId3, LobbyAttribute2Value2);
+			LobbySnapshot.Members = { User1, User2 };
+
+			TMap<FAccountId, ELobbyMemberLeaveReason> LeaveReasons;
+			TOnlineResult<FLobbyClientDataPrepareServiceSnapshot> PrepareResult = ClientData.PrepareServiceSnapshot({ MoveTemp(LobbySnapshot), MoveTemp(LobbyMemberSnapshots), MoveTemp(LeaveReasons) });
+			UTEST_TRUE("Check PrepareServiceSnapshot succeeded.", PrepareResult.IsOk());
+
+			FLobbyClientDataCommitServiceSnapshot::Result CommitResult = ClientData.CommitServiceSnapshot({ &LobbyEvents });
+			UTEST_EQUAL("Check no local members left.", CommitResult.LeavingLocalMembers.Num(), 0);
+		}
+
+		{
+			// 1.3.1. Snapshot with changed lobby attribute does not trigger lobby attribute changed.
+			// 1.3.2. Snapshot with changed lobby member attribute does not trigger lobby member attribute changed.
+			UTEST_EQUAL("Snapshot does not trigger any events.", EventCapture.GetTotalNotificationsReceived(), 0);
+			UTEST_EQUAL("Verify lobby owner.", ClientData.GetPublicData().OwnerAccountId, User2);
+			UTEST_EQUAL("Verify lobby attributes.", ClientData.GetPublicData().Attributes.Num(), 2);
+			const FSchemaVariant* LobbyAttribute1 = ClientData.GetPublicData().Attributes.Find(LobbyAttributeId1);
+			UTEST_NOT_NULL("Verify lobby attributes.", LobbyAttribute1);
+			UTEST_EQUAL("Verify lobby attributes.", *LobbyAttribute1, LobbyAttribute1Value1);
+			const FSchemaVariant* LobbyAttribute2 = ClientData.GetPublicData().Attributes.Find(LobbyAttributeId2);
+			UTEST_NOT_NULL("Verify lobby attributes.", LobbyAttribute2);
+			UTEST_EQUAL("Verify lobby attributes.", *LobbyAttribute2, LobbyAttribute2Value2);
+			UTEST_EQUAL("Verify lobby members.", ClientData.GetPublicData().Members.Num(), 2);
+			const TSharedRef<const FLobbyMember>* User1Data = ClientData.GetPublicData().Members.Find(User1);
+			UTEST_NOT_NULL("Verify lobby members.", User1Data);
+			UTEST_EQUAL("Verify lobby member attributes.", (*User1Data)->Attributes.Num(), 2);
+			const FSchemaVariant* LobbyMemberAttribute1 = (*User1Data)->Attributes.Find(LobbyMemberAttributeId1);
+			UTEST_NOT_NULL("Verify lobby member attributes.", LobbyMemberAttribute1);
+			UTEST_EQUAL("Verify lobby member attributes.", *LobbyMemberAttribute1, LobbyMemberAttribute1Value1);
+			const FSchemaVariant* LobbyMemberAttribute2 = (*User1Data)->Attributes.Find(LobbyMemberAttributeId2);
+			UTEST_NOT_NULL("Verify lobby member attributes.", LobbyMemberAttribute2);
+			UTEST_EQUAL("Verify lobby member attributes.", *LobbyMemberAttribute2, LobbyMemberAttribute2Value2);
+			const TSharedRef<const FLobbyMember>* User2Data = ClientData.GetPublicData().Members.Find(User2);
+			UTEST_NOT_NULL("Verify lobby members.", User2Data);
+			UTEST_EQUAL("Verify lobby member attributes.", (*User2Data)->Attributes.Num(), 0);
+		}
 
-			FClientLobbySnapshot LobbySnapshot;
+		// 1.3 setup. Remove lobby and lobby member attribute.
+		{
+			TMap<FAccountId, FLobbyMemberServiceSnapshot> LobbyMemberSnapshots;
+			{
+				FLobbyMemberServiceSnapshot& MemberSnapshot = LobbyMemberSnapshots.Add(User1);
+				MemberSnapshot.AccountId = User1;
+				MemberSnapshot.SchemaServiceSnapshot.Attributes.Add(LobbyMemberServiceAttributeId1, LobbyMemberAttribute1Value1);
+			}
+			{
+				FLobbyMemberServiceSnapshot& MemberSnapshot = LobbyMemberSnapshots.Add(User2);
+				MemberSnapshot.AccountId = User2;
+			}
+
+			FLobbyServiceSnapshot LobbySnapshot;
+			LobbySnapshot.OwnerAccountId = User2;
+			LobbySnapshot.MaxMembers = LobbyMaxMembers;
+			LobbySnapshot.JoinPolicy = LobbyJoinPolicy;
+			LobbySnapshot.SchemaServiceSnapshot.Attributes.Add(LobbyServiceAttributeId1, SchemaCompatibilityAttributeValue);
+			LobbySnapshot.SchemaServiceSnapshot.Attributes.Add(LobbyServiceAttributeId2, LobbyAttribute1Value1);
+			LobbySnapshot.Members = { User1, User2 };
+
+			TMap<FAccountId, ELobbyMemberLeaveReason> LeaveReasons;
+			TOnlineResult<FLobbyClientDataPrepareServiceSnapshot> PrepareResult = ClientData.PrepareServiceSnapshot({ MoveTemp(LobbySnapshot), MoveTemp(LobbyMemberSnapshots), MoveTemp(LeaveReasons) });
+			UTEST_TRUE("Check PrepareServiceSnapshot succeeded.", PrepareResult.IsOk());
+
+			FLobbyClientDataCommitServiceSnapshot::Result CommitResult = ClientData.CommitServiceSnapshot({ &LobbyEvents });
+			UTEST_EQUAL("Check no local members left.", CommitResult.LeavingLocalMembers.Num(), 0);
+		}
+
+		{
+			// 1.3.1. Snapshot with changed lobby attribute does not trigger lobby attribute changed.
+			// 1.3.2. Snapshot with changed lobby member attribute does not trigger lobby member attribute changed.
+			UTEST_EQUAL("Snapshot does not trigger any events.", EventCapture.GetTotalNotificationsReceived(), 0);
+			UTEST_EQUAL("Verify lobby owner.", ClientData.GetPublicData().OwnerAccountId, User2);
+			UTEST_EQUAL("Verify lobby attributes.", ClientData.GetPublicData().Attributes.Num(), 1);
+			const FSchemaVariant* LobbyAttribute1 = ClientData.GetPublicData().Attributes.Find(LobbyAttributeId1);
+			UTEST_NOT_NULL("Verify lobby attributes.", LobbyAttribute1);
+			UTEST_EQUAL("Verify lobby attributes.", *LobbyAttribute1, LobbyAttribute1Value1);
+			const TSharedRef<const FLobbyMember>* User1Data = ClientData.GetPublicData().Members.Find(User1);
+			UTEST_NOT_NULL("Verify lobby members.", User1Data);
+			UTEST_EQUAL("Verify lobby member attributes.", (*User1Data)->Attributes.Num(), 1);
+			const FSchemaVariant* LobbyMemberAttribute1 = (*User1Data)->Attributes.Find(LobbyMemberAttributeId1);
+			UTEST_NOT_NULL("Verify lobby member attributes.", LobbyMemberAttribute1);
+			UTEST_EQUAL("Verify lobby member attributes.", *LobbyMemberAttribute1, LobbyMemberAttribute1Value1);
+			const TSharedRef<const FLobbyMember>* User2Data = ClientData.GetPublicData().Members.Find(User2);
+			UTEST_NOT_NULL("Verify lobby members.", User2Data);
+			UTEST_EQUAL("Verify lobby member attributes.", (*User2Data)->Attributes.Num(), 0);
+		}
+
+		// 1.3 setup. User leaving the lobby.
+		{
+			TMap<FAccountId, FLobbyMemberServiceSnapshot> LobbyMemberSnapshots;
+			{
+				FLobbyMemberServiceSnapshot& MemberSnapshot = LobbyMemberSnapshots.Add(User1);
+				MemberSnapshot.AccountId = User1;
+			}
+
+			FLobbyServiceSnapshot LobbySnapshot;
 			LobbySnapshot.OwnerAccountId = User1;
-			LobbySnapshot.SchemaName = LobbySchemaName;
 			LobbySnapshot.MaxMembers = LobbyMaxMembers;
 			LobbySnapshot.JoinPolicy = LobbyJoinPolicy;
 			LobbySnapshot.Members = {User1};
 
 			TMap<FAccountId, ELobbyMemberLeaveReason> LeaveReasons;
-			FApplyLobbyUpdateResult Result = ClientData.ApplyLobbyUpdateFromServiceSnapshot(MoveTemp(LobbySnapshot), MoveTemp(LobbyMemberSnapshots), MoveTemp(LeaveReasons), &LobbyEvents);
-			UTEST_EQUAL("Check no local members left.", Result.LeavingLocalMembers.Num(), 0);
-		}
-
-		// 1.3.1. Snapshot with cleared lobby attribute does not trigger lobby attribute changed.
-		// 1.3.2. Snapshot with cleared lobby member attribute does not trigger lobby member attribute changed.
-		// 1.3.3. Snapshot with members leaving does not trigger lobby member left or lobby left.
-		UTEST_EQUAL("Snapshot does not trigger any events.", EventCapture.GetTotalNotificationsReceived(), 0);
-		UTEST_EQUAL("Verify lobby attributes.", ClientData.GetPublicData().Attributes.Num(), 0);
-		UTEST_EQUAL("Verify lobby members.", ClientData.GetPublicData().Members.Num(), 1);
-		UTEST_NOT_NULL("Verify lobby members.", ClientData.GetPublicData().Members.Find(User1));
-		UTEST_EQUAL("Verify lobby member attributes.", ClientData.GetPublicData().Members[User1]->Attributes.Num(), 0);
+			TOnlineResult<FLobbyClientDataPrepareServiceSnapshot> PrepareResult = ClientData.PrepareServiceSnapshot({ MoveTemp(LobbySnapshot), MoveTemp(LobbyMemberSnapshots), MoveTemp(LeaveReasons) });
+			UTEST_TRUE("Check PrepareServiceSnapshot succeeded.", PrepareResult.IsOk());
+
+			FLobbyClientDataCommitServiceSnapshot::Result CommitResult = ClientData.CommitServiceSnapshot({ &LobbyEvents });
+			UTEST_EQUAL("Check no local members left.", CommitResult.LeavingLocalMembers.Num(), 0);
+		}
+
+		{
+			// 1.3.3. Snapshot with members leaving does not trigger lobby member left or lobby left.
+			UTEST_EQUAL("Snapshot does not trigger any events.", EventCapture.GetTotalNotificationsReceived(), 0);
+			UTEST_EQUAL("Verify lobby attributes.", ClientData.GetPublicData().Attributes.Num(), 0);
+			UTEST_EQUAL("Verify lobby members.", ClientData.GetPublicData().Members.Num(), 1);
+			UTEST_NOT_NULL("Verify lobby members.", ClientData.GetPublicData().Members.Find(User1));
+			UTEST_EQUAL("Verify lobby member attributes.", ClientData.GetPublicData().Members[User1]->Attributes.Num(), 0);
+		}
 	}
 
 	// 2. Applying snapshot with local members present does fire events.
@@ -268,186 +483,307 @@ bool FClientLobbyDataTest::RunTest(const FString& Parameters)
 		FLobbyEvents LobbyEvents;
 		FLobbyEventCapture EventCapture(LobbyEvents);
 		FLobbyId LobbyId(EOnlineServices::Null, 1);
-		FClientLobbyData ClientData(LobbyId);
+		FLobbyClientData ClientData(LobbyId, SchemaRegistry);
 
 		FAccountId User1(EOnlineServices::Null, 1);
 		FAccountId User2(EOnlineServices::Null, 2);
-		FName LobbySchemaName = TEXT("SchemaName");
 		int32 LobbyMaxMembers = 5;
 		ELobbyJoinPolicy LobbyJoinPolicy = ELobbyJoinPolicy::PublicAdvertised;
-		FLobbyAttributeId LobbyAttribute1Key = TEXT("LobbyAttribute1");
-		FLobbyVariant LobbyAttribute1Value;
-		LobbyAttribute1Value.Set(TEXT("Attribute1"));
-		FLobbyAttributeId LobbyAttribute2Key = TEXT("LobbyAttribute2");
-		FLobbyVariant LobbyAttribute2Value;
-		LobbyAttribute2Value.Set(TEXT("Attribute2"));
-		FLobbyAttributeId LobbyMemberAttribute1Key = TEXT("LobbyMemberAttribute1");
-		FLobbyVariant LobbyMemberAttribute1Value;
-		LobbyMemberAttribute1Value.Set(TEXT("MemberAttribute1"));
-		FLobbyAttributeId LobbyMemberAttribute2Key = TEXT("LobbyMemberAttribute2");
-		FLobbyVariant LobbyMemberAttribute2Value;
-		LobbyMemberAttribute2Value.Set(TEXT("MemberAttribute2"));
-
-		// 2.1 setup
-		{
-			TMap<FAccountId, TSharedRef<FClientLobbyMemberSnapshot>> LobbyMemberSnapshots;
+
+		// 2.1 setup - initial lobby setup.
+		{
+			TMap<FAccountId, FLobbyMemberServiceSnapshot> LobbyMemberSnapshots;
 			{
-				TSharedRef<FClientLobbyMemberSnapshot> MemberSnapshot = MakeShared<FClientLobbyMemberSnapshot>();
-				MemberSnapshot->AccountId = User1;
-				LobbyMemberSnapshots.Add(User1, MemberSnapshot);
+				FLobbyMemberServiceSnapshot& MemberSnapshot = LobbyMemberSnapshots.Add(User1);
+				MemberSnapshot.AccountId = User1;
+				MemberSnapshot.SchemaServiceSnapshot.Attributes.Add(LobbyMemberServiceAttributeId1, LobbyMemberAttribute1Value1);
 			}
 			{
-				TSharedRef<FClientLobbyMemberSnapshot> MemberSnapshot = MakeShared<FClientLobbyMemberSnapshot>();
-				MemberSnapshot->AccountId = User2;
-				MemberSnapshot->Attributes.Add(LobbyMemberAttribute1Key, LobbyMemberAttribute1Value);
-				LobbyMemberSnapshots.Add(User2, MemberSnapshot);
+				FLobbyMemberServiceSnapshot& MemberSnapshot = LobbyMemberSnapshots.Add(User2);
+				MemberSnapshot.AccountId = User2;
 			}
 
-			FClientLobbySnapshot LobbySnapshot;
+			FLobbyServiceSnapshot LobbySnapshot;
 			LobbySnapshot.OwnerAccountId = User1;
-			LobbySnapshot.SchemaName = LobbySchemaName;
 			LobbySnapshot.MaxMembers = LobbyMaxMembers;
 			LobbySnapshot.JoinPolicy = LobbyJoinPolicy;
-			LobbySnapshot.Attributes.Add(LobbyAttribute1Key, LobbyAttribute1Value);
+			LobbySnapshot.SchemaServiceSnapshot.Attributes.Add(LobbyServiceAttributeId1, SchemaCompatibilityAttributeValue);
+			LobbySnapshot.SchemaServiceSnapshot.Attributes.Add(LobbyServiceAttributeId2, LobbyAttribute1Value1);
 			LobbySnapshot.Members = {User1, User2};
 
 			TMap<FAccountId, ELobbyMemberLeaveReason> LeaveReasons;
-			FApplyLobbyUpdateResult Result = ClientData.ApplyLobbyUpdateFromServiceSnapshot(MoveTemp(LobbySnapshot), MoveTemp(LobbyMemberSnapshots), MoveTemp(LeaveReasons), &LobbyEvents);
-			UTEST_EQUAL("Check no members left.", Result.LeavingLocalMembers.Num(), 0);
+			TOnlineResult<FLobbyClientDataPrepareServiceSnapshot> PrepareResult = ClientData.PrepareServiceSnapshot({ MoveTemp(LobbySnapshot), MoveTemp(LobbyMemberSnapshots), MoveTemp(LeaveReasons) });
+			UTEST_TRUE("Check PrepareServiceSnapshot succeeded.", PrepareResult.IsOk());
+
+			FLobbyClientDataCommitServiceSnapshot::Result CommitResult = ClientData.CommitServiceSnapshot({ &LobbyEvents });
+			UTEST_EQUAL("Check no local members left.", CommitResult.LeavingLocalMembers.Num(), 0);
+		}
+
+		{
+			// 2.1.1. Initial snapshot does not trigger lobby joined or lobby member joined. - Member in snapshot is local member to simulate lobby creation.
+			UTEST_EQUAL("Snapshot does not trigger any events.", EventCapture.GetTotalNotificationsReceived(), 0);
+			UTEST_EQUAL("Verify lobby members.", ClientData.GetPublicData().Members.Num(), 2);
+			UTEST_NOT_NULL("Verify lobby members.", ClientData.GetPublicData().Members.Find(User1));
+			UTEST_NOT_NULL("Verify lobby members.", ClientData.GetPublicData().Members.Find(User2));
+		}
+
+		// 2.2 setup - add local lobby member User1.
+		{
+			EventCapture.Empty();
+
+			FLobbyClientChanges LobbyChanges;
+			LobbyChanges.MemberAttributes = FLobbyClientMemberAttributeChanges();
+			TOnlineResult<FLobbyClientDataPrepareClientChanges> PrepareResult = ClientData.PrepareClientChanges({ User1, MoveTemp(LobbyChanges) });
+			UTEST_TRUE("Check PrepareClientChanges succeeded.", PrepareResult.IsOk());
+
+			FLobbyClientDataCommitClientChanges::Result CommitResult = ClientData.CommitClientChanges({ &LobbyEvents });
+			UTEST_EQUAL("Check no local members left.", CommitResult.LeavingLocalMembers.Num(), 0);
 		}
 
-		// 2.1.1. Initial snapshot does not trigger lobby joined or lobby member joined. - Member in snapshot is local member to simulate lobby creation.
-		UTEST_EQUAL("Snapshot does not trigger any events.", EventCapture.GetTotalNotificationsReceived(), 0);
-		UTEST_EQUAL("Verify lobby members.", ClientData.GetPublicData().Members.Num(), 2);
-		UTEST_NOT_NULL("Verify lobby members.", ClientData.GetPublicData().Members.Find(User1));
-		UTEST_NOT_NULL("Verify lobby members.", ClientData.GetPublicData().Members.Find(User2));
+		{
+			// 2.2.1. Local member added to lobby through local change. Adding a local member will cause events to begin triggering.
+			UTEST_EQUAL("Check lobby and member joined events were received.", EventCapture.GetTotalNotificationsReceived(), 3);
+			UTEST_EQUAL("Check lobby joined event received.", EventCapture.LobbyJoined.Num(), 1);
+			UTEST_EQUAL("Check lobby joined event received first.", EventCapture.LobbyJoined[0].GlobalIndex, 0);
+			UTEST_EQUAL("Check lobby joined event is valid.", EventCapture.LobbyJoined[0].Notification.Lobby, ClientData.GetPublicDataPtr());
+			UTEST_EQUAL("Check member joined events received.", EventCapture.LobbyMemberJoined.Num(), 2);
+			UTEST_EQUAL("Check member joined event is valid.", EventCapture.LobbyMemberJoined[0].Notification.Lobby, ClientData.GetPublicDataPtr());
+			UTEST_TRUE("Check member joined event is valid.", EventCapture.LobbyMemberJoined[0].Notification.Member->AccountId.IsValid());
+		}
 
-		// 2.2 setup
+		// 2.3 setup - Change lobby leader.
 		{
 			EventCapture.Empty();
 
-			FClientLobbyDataChanges LobbyChanges;
-			LobbyChanges.MutatedMembers.Add(User1, MakeShared<FClientLobbyMemberDataChanges>());
-			FApplyLobbyUpdateResult Result = ClientData.ApplyLobbyUpdateFromLocalChanges(MoveTemp(LobbyChanges), LobbyEvents);
-			UTEST_EQUAL("Check no local members left.", Result.LeavingLocalMembers.Num(), 0);
+			TMap<FAccountId, FLobbyMemberServiceSnapshot> LobbyMemberSnapshots;
+			{
+				FLobbyMemberServiceSnapshot& MemberSnapshot = LobbyMemberSnapshots.Add(User1);
+				MemberSnapshot.AccountId = User1;
+				MemberSnapshot.SchemaServiceSnapshot.Attributes.Add(LobbyMemberServiceAttributeId1, LobbyMemberAttribute1Value1);
+			}
+
+			FLobbyServiceSnapshot LobbySnapshot;
+			LobbySnapshot.OwnerAccountId = User2;
+			LobbySnapshot.MaxMembers = LobbyMaxMembers;
+			LobbySnapshot.JoinPolicy = LobbyJoinPolicy;
+			LobbySnapshot.SchemaServiceSnapshot.Attributes.Add(LobbyServiceAttributeId1, SchemaCompatibilityAttributeValue);
+			LobbySnapshot.SchemaServiceSnapshot.Attributes.Add(LobbyServiceAttributeId2, LobbyAttribute1Value1);
+			LobbySnapshot.Members = {User1, User2};
+
+			TMap<FAccountId, ELobbyMemberLeaveReason> LeaveReasons;
+			TOnlineResult<FLobbyClientDataPrepareServiceSnapshot> PrepareResult = ClientData.PrepareServiceSnapshot({ MoveTemp(LobbySnapshot), MoveTemp(LobbyMemberSnapshots), MoveTemp(LeaveReasons) });
+			UTEST_TRUE("Check PrepareServiceSnapshot succeeded.", PrepareResult.IsOk());
+
+			FLobbyClientDataCommitServiceSnapshot::Result CommitResult = ClientData.CommitServiceSnapshot({ &LobbyEvents });
+			UTEST_EQUAL("Check no local members left.", CommitResult.LeavingLocalMembers.Num(), 0);
 		}
 
-		// 2.2.1. Local member added to lobby through local change. Adding a local member will cause events to begin triggering.
-		UTEST_EQUAL("Check lobby and member joined events were received.", EventCapture.GetTotalNotificationsReceived(), 3);
-		UTEST_EQUAL("Check lobby joined event received.", EventCapture.LobbyJoined.Num(), 1);
-		UTEST_EQUAL("Check lobby joined event received first.", EventCapture.LobbyJoined[0].GlobalIndex, 0);
-		UTEST_EQUAL("Check lobby joined event is valid.", EventCapture.LobbyJoined[0].Notification.Lobby, ClientData.GetPublicDataPtr());
-		UTEST_EQUAL("Check member joined events received.", EventCapture.LobbyMemberJoined.Num(), 2);
-		UTEST_EQUAL("Check member joined event is valid.", EventCapture.LobbyMemberJoined[0].Notification.Lobby, ClientData.GetPublicDataPtr());
-		UTEST_TRUE("Check member joined event is valid.", EventCapture.LobbyMemberJoined[0].Notification.Member->AccountId.IsValid());
+		{
+			// 2.3.1. Snapshot with changed lobby leader triggers lobby leader changed.
+			UTEST_EQUAL("Check events received.", EventCapture.GetTotalNotificationsReceived(), 1);
+			UTEST_EQUAL("Check LobbySchemaChanged is valid.", EventCapture.LobbySchemaChanged.Num(), 0);
+			UTEST_EQUAL("Check LobbyLeaderChanged is valid.", EventCapture.LobbyLeaderChanged.Num(), 1);
+			UTEST_EQUAL("Check LobbyLeaderChanged is valid.", EventCapture.LobbyLeaderChanged[0].Notification.Lobby, ClientData.GetPublicDataPtr());
+			UTEST_EQUAL("Check LobbyLeaderChanged is valid.", EventCapture.LobbyLeaderChanged[0].Notification.Leader->AccountId, User2);
+		}
 
-		// 2.3 setup
+		// 2.4 setup - Add lobby and member attribute
 		{
 			EventCapture.Empty();
 
-			LobbyAttribute1Value.Set(TEXT("ModifiedAttribute1"));
-			LobbyMemberAttribute1Value.Set(TEXT("ModifiedMemberAttribute1"));
-			LobbySchemaName = TEXT("ModifiedSchemaName");
+			TMap<FAccountId, FLobbyMemberServiceSnapshot> LobbyMemberSnapshots;
+			{
+				FLobbyMemberServiceSnapshot& MemberSnapshot = LobbyMemberSnapshots.Add(User1);
+				MemberSnapshot.AccountId = User1;
+				MemberSnapshot.SchemaServiceSnapshot.Attributes.Add(LobbyMemberServiceAttributeId1, LobbyMemberAttribute1Value1);
+				MemberSnapshot.SchemaServiceSnapshot.Attributes.Add(LobbyMemberServiceAttributeId2, LobbyMemberAttribute2Value1);
+			}
+
+			FLobbyServiceSnapshot LobbySnapshot;
+			LobbySnapshot.OwnerAccountId = User2;
+			LobbySnapshot.MaxMembers = LobbyMaxMembers;
+			LobbySnapshot.JoinPolicy = LobbyJoinPolicy;
+			LobbySnapshot.SchemaServiceSnapshot.Attributes.Add(LobbyServiceAttributeId1, SchemaCompatibilityAttributeValue);
+			LobbySnapshot.SchemaServiceSnapshot.Attributes.Add(LobbyServiceAttributeId2, LobbyAttribute1Value1);
+			LobbySnapshot.SchemaServiceSnapshot.Attributes.Add(LobbyServiceAttributeId3, LobbyAttribute2Value1);
+			LobbySnapshot.Members = { User1, User2 };
+
+			TMap<FAccountId, ELobbyMemberLeaveReason> LeaveReasons;
+			TOnlineResult<FLobbyClientDataPrepareServiceSnapshot> PrepareResult = ClientData.PrepareServiceSnapshot({ MoveTemp(LobbySnapshot), MoveTemp(LobbyMemberSnapshots), MoveTemp(LeaveReasons) });
+			UTEST_TRUE("Check PrepareServiceSnapshot succeeded.", PrepareResult.IsOk());
+
+			FLobbyClientDataCommitServiceSnapshot::Result CommitResult = ClientData.CommitServiceSnapshot({ &LobbyEvents });
+			UTEST_EQUAL("Check no local members left.", CommitResult.LeavingLocalMembers.Num(), 0);
+		}
+
+		{
+			// 2.4.1. Snapshot with changed lobby attribute triggers lobby attribute changed for added attribute.
+			// 2.4.2. Snapshot with changed lobby member attribute triggers lobby member attribute changed for added attribute.
+			UTEST_EQUAL("Check events received.", EventCapture.GetTotalNotificationsReceived(), 2);
+			UTEST_EQUAL("Check LobbyAttributesChanged is valid.", EventCapture.LobbyAttributesChanged.Num(), 1);
+			UTEST_EQUAL("Check LobbyAttributesChanged is valid.", EventCapture.LobbyAttributesChanged[0].Notification.Lobby, ClientData.GetPublicDataPtr());
+			UTEST_EQUAL("Check LobbyAttributesChanged is valid.", EventCapture.LobbyAttributesChanged[0].Notification.AddedAttributes.Num(), 1);
+			UTEST_EQUAL("Check LobbyAttributesChanged is valid.", EventCapture.LobbyAttributesChanged[0].Notification.ChangedAttributes.Num(), 0);
+			UTEST_EQUAL("Check LobbyAttributesChanged is valid.", EventCapture.LobbyAttributesChanged[0].Notification.RemovedAttributes.Num(), 0);
+			const FSchemaVariant* LobbyAttribute2 = EventCapture.LobbyAttributesChanged[0].Notification.AddedAttributes.Find(LobbyAttributeId2);
+			UTEST_NOT_NULL("Check LobbyAttributesChanged is valid.", LobbyAttribute2);
+			UTEST_EQUAL("Check LobbyAttributesChanged is valid.", *LobbyAttribute2, LobbyAttribute2Value1);
+			UTEST_EQUAL("Check LobbyMemberAttributesChanged is valid.", EventCapture.LobbyMemberAttributesChanged.Num(), 1);
+			UTEST_EQUAL("Check LobbyMemberAttributesChanged is valid.", EventCapture.LobbyMemberAttributesChanged[0].Notification.Lobby, ClientData.GetPublicDataPtr());
+			UTEST_EQUAL("Check LobbyMemberAttributesChanged is valid.", EventCapture.LobbyMemberAttributesChanged[0].Notification.Member->AccountId, User1);
+			UTEST_EQUAL("Check LobbyMemberAttributesChanged is valid.", EventCapture.LobbyMemberAttributesChanged[0].Notification.AddedAttributes.Num(), 1);
+			UTEST_EQUAL("Check LobbyMemberAttributesChanged is valid.", EventCapture.LobbyMemberAttributesChanged[0].Notification.ChangedAttributes.Num(), 0);
+			UTEST_EQUAL("Check LobbyMemberAttributesChanged is valid.", EventCapture.LobbyMemberAttributesChanged[0].Notification.RemovedAttributes.Num(), 0);
+			const FSchemaVariant* LobbyMemberAttribute2 = EventCapture.LobbyMemberAttributesChanged[0].Notification.AddedAttributes.Find(LobbyMemberAttributeId2);
+			UTEST_NOT_NULL("Check LobbyMemberAttributesChanged is valid.", LobbyMemberAttribute2);
+			UTEST_EQUAL("Check LobbyMemberAttributesChanged is valid.", *LobbyMemberAttribute2, LobbyMemberAttribute2Value1);
+		}
 
-			TMap<FAccountId, TSharedRef<FClientLobbyMemberSnapshot>> LobbyMemberSnapshots;
+		// 2.5 setup - Modify lobby and member attribute
+		{
+			EventCapture.Empty();
+
+			TMap<FAccountId, FLobbyMemberServiceSnapshot> LobbyMemberSnapshots;
 			{
-				TSharedRef<FClientLobbyMemberSnapshot> MemberSnapshot = MakeShared<FClientLobbyMemberSnapshot>();
-				MemberSnapshot->AccountId = User1;
-				LobbyMemberSnapshots.Add(User1, MemberSnapshot);
+				FLobbyMemberServiceSnapshot& MemberSnapshot = LobbyMemberSnapshots.Add(User1);
+				MemberSnapshot.AccountId = User1;
+				MemberSnapshot.SchemaServiceSnapshot.Attributes.Add(LobbyMemberServiceAttributeId1, LobbyMemberAttribute1Value1);
+				MemberSnapshot.SchemaServiceSnapshot.Attributes.Add(LobbyMemberServiceAttributeId2, LobbyMemberAttribute2Value2);
 			}
+
+			FLobbyServiceSnapshot LobbySnapshot;
+			LobbySnapshot.OwnerAccountId = User2;
+			LobbySnapshot.MaxMembers = LobbyMaxMembers;
+			LobbySnapshot.JoinPolicy = LobbyJoinPolicy;
+			LobbySnapshot.SchemaServiceSnapshot.Attributes.Add(LobbyServiceAttributeId1, SchemaCompatibilityAttributeValue);
+			LobbySnapshot.SchemaServiceSnapshot.Attributes.Add(LobbyServiceAttributeId2, LobbyAttribute1Value1);
+			LobbySnapshot.SchemaServiceSnapshot.Attributes.Add(LobbyServiceAttributeId3, LobbyAttribute2Value2);
+			LobbySnapshot.Members = { User1, User2 };
+
+			TMap<FAccountId, ELobbyMemberLeaveReason> LeaveReasons;
+			TOnlineResult<FLobbyClientDataPrepareServiceSnapshot> PrepareResult = ClientData.PrepareServiceSnapshot({ MoveTemp(LobbySnapshot), MoveTemp(LobbyMemberSnapshots), MoveTemp(LeaveReasons) });
+			UTEST_TRUE("Check PrepareServiceSnapshot succeeded.", PrepareResult.IsOk());
+
+			FLobbyClientDataCommitServiceSnapshot::Result CommitResult = ClientData.CommitServiceSnapshot({ &LobbyEvents });
+			UTEST_EQUAL("Check no local members left.", CommitResult.LeavingLocalMembers.Num(), 0);
+		}
+
+		{
+			// 2.5.1. Snapshot with changed lobby attribute triggers lobby attribute changed for changed attribute.
+			// 2.5.2. Snapshot with changed lobby member attribute triggers lobby member attribute changed for changed attribute.
+			UTEST_EQUAL("Check events received.", EventCapture.GetTotalNotificationsReceived(), 2);
+			UTEST_EQUAL("Check LobbyAttributesChanged is valid.", EventCapture.LobbyAttributesChanged.Num(), 1);
+			UTEST_EQUAL("Check LobbyAttributesChanged is valid.", EventCapture.LobbyAttributesChanged[0].Notification.Lobby, ClientData.GetPublicDataPtr());
+			UTEST_EQUAL("Check LobbyAttributesChanged is valid.", EventCapture.LobbyAttributesChanged[0].Notification.AddedAttributes.Num(), 0);
+			UTEST_EQUAL("Check LobbyAttributesChanged is valid.", EventCapture.LobbyAttributesChanged[0].Notification.ChangedAttributes.Num(), 1);
+			UTEST_EQUAL("Check LobbyAttributesChanged is valid.", EventCapture.LobbyAttributesChanged[0].Notification.RemovedAttributes.Num(), 0);
+			const TPair<FSchemaVariant, FSchemaVariant>* LobbyAttribute2 = EventCapture.LobbyAttributesChanged[0].Notification.ChangedAttributes.Find(LobbyAttributeId2);
+			UTEST_NOT_NULL("Check LobbyAttributesChanged is valid.", LobbyAttribute2);
+			UTEST_EQUAL("Check LobbyAttributesChanged is valid.", *LobbyAttribute2, (TPair<FSchemaVariant, FSchemaVariant>(LobbyAttribute2Value1, LobbyAttribute2Value2)));
+			UTEST_EQUAL("Check LobbyMemberAttributesChanged is valid.", EventCapture.LobbyMemberAttributesChanged.Num(), 1);
+			UTEST_EQUAL("Check LobbyMemberAttributesChanged is valid.", EventCapture.LobbyMemberAttributesChanged[0].Notification.Lobby, ClientData.GetPublicDataPtr());
+			UTEST_EQUAL("Check LobbyMemberAttributesChanged is valid.", EventCapture.LobbyMemberAttributesChanged[0].Notification.Member->AccountId, User1);
+			UTEST_EQUAL("Check LobbyMemberAttributesChanged is valid.", EventCapture.LobbyMemberAttributesChanged[0].Notification.AddedAttributes.Num(), 0);
+			UTEST_EQUAL("Check LobbyMemberAttributesChanged is valid.", EventCapture.LobbyMemberAttributesChanged[0].Notification.ChangedAttributes.Num(), 1);
+			UTEST_EQUAL("Check LobbyMemberAttributesChanged is valid.", EventCapture.LobbyMemberAttributesChanged[0].Notification.RemovedAttributes.Num(), 0);
+			const TPair<FSchemaVariant, FSchemaVariant>* LobbyMemberAttribute2 = EventCapture.LobbyMemberAttributesChanged[0].Notification.ChangedAttributes.Find(LobbyMemberAttributeId2);
+			UTEST_NOT_NULL("Check LobbyMemberAttributesChanged is valid.", LobbyMemberAttribute2);
+			UTEST_EQUAL("Check LobbyMemberAttributesChanged is valid.", *LobbyMemberAttribute2, (TPair<FSchemaVariant, FSchemaVariant>(LobbyMemberAttribute2Value1, LobbyMemberAttribute2Value2)));
+		}
+
+		// 2.6 setup - Clear lobby and member attribute
+		{
+			EventCapture.Empty();
+
+			TMap<FAccountId, FLobbyMemberServiceSnapshot> LobbyMemberSnapshots;
 			{
-				TSharedRef<FClientLobbyMemberSnapshot> MemberSnapshot = MakeShared<FClientLobbyMemberSnapshot>();
-				MemberSnapshot->AccountId = User2;
-				MemberSnapshot->Attributes.Add(LobbyMemberAttribute1Key, LobbyMemberAttribute1Value);
-				LobbyMemberSnapshots.Add(User2, MemberSnapshot);
+				FLobbyMemberServiceSnapshot& MemberSnapshot = LobbyMemberSnapshots.Add(User1);
+				MemberSnapshot.AccountId = User1;
+				MemberSnapshot.SchemaServiceSnapshot.Attributes.Add(LobbyMemberServiceAttributeId1, LobbyMemberAttribute1Value1);
 			}
 
-			FClientLobbySnapshot LobbySnapshot;
+			FLobbyServiceSnapshot LobbySnapshot;
 			LobbySnapshot.OwnerAccountId = User2;
-			LobbySnapshot.SchemaName = LobbySchemaName;
 			LobbySnapshot.MaxMembers = LobbyMaxMembers;
 			LobbySnapshot.JoinPolicy = LobbyJoinPolicy;
-			LobbySnapshot.Attributes.Add(LobbyAttribute1Key, LobbyAttribute1Value);
-			LobbySnapshot.Members = {User1, User2};
+			LobbySnapshot.SchemaServiceSnapshot.Attributes.Add(LobbyServiceAttributeId1, SchemaCompatibilityAttributeValue);
+			LobbySnapshot.SchemaServiceSnapshot.Attributes.Add(LobbyServiceAttributeId2, LobbyAttribute1Value1);
+			LobbySnapshot.Members = { User1, User2 };
 
 			TMap<FAccountId, ELobbyMemberLeaveReason> LeaveReasons;
-			FApplyLobbyUpdateResult Result = ClientData.ApplyLobbyUpdateFromServiceSnapshot(MoveTemp(LobbySnapshot), MoveTemp(LobbyMemberSnapshots), MoveTemp(LeaveReasons), &LobbyEvents);
-			UTEST_EQUAL("Check no local members left.", Result.LeavingLocalMembers.Num(), 0);
-		}
-
-		// 2.3.1. Snapshot with changed lobby attribute triggers lobby attribute changed.
-		// 2.3.2. Snapshot with changed lobby member attribute triggers lobby member attribute changed.
-		// 2.3.3. Snapshot with changed lobby schema triggers lobby schema changed.
-		// 2.3.4. Snapshot with changed lobby leader triggers lobby leader changed.
-		UTEST_EQUAL("Check events received.", EventCapture.GetTotalNotificationsReceived(), 4);
-		UTEST_EQUAL("Check LobbyAttributesChanged is valid.", EventCapture.LobbyAttributesChanged.Num(), 1);
-		UTEST_EQUAL("Check LobbyAttributesChanged is valid.", EventCapture.LobbyAttributesChanged[0].Notification.Lobby, ClientData.GetPublicDataPtr());
-		UTEST_EQUAL("Check LobbyAttributesChanged is valid.", EventCapture.LobbyAttributesChanged[0].Notification.ChangedAttributes.Num(), 1);
-		UTEST_NOT_NULL("Check LobbyAttributesChanged is valid.", EventCapture.LobbyAttributesChanged[0].Notification.ChangedAttributes.Find(LobbyAttribute1Key));
-		UTEST_EQUAL("Check LobbyMemberAttributesChanged is valid.", EventCapture.LobbyMemberAttributesChanged.Num(), 1);
-		UTEST_EQUAL("Check LobbyMemberAttributesChanged is valid.", EventCapture.LobbyMemberAttributesChanged[0].Notification.Lobby, ClientData.GetPublicDataPtr());
-		UTEST_EQUAL("Check LobbyMemberAttributesChanged is valid.", EventCapture.LobbyMemberAttributesChanged[0].Notification.Member->AccountId, User2);
-		UTEST_EQUAL("Check LobbyMemberAttributesChanged is valid.", EventCapture.LobbyMemberAttributesChanged[0].Notification.ChangedAttributes.Num(), 1);
-		UTEST_NOT_NULL("Check LobbyMemberAttributesChanged is valid.", EventCapture.LobbyMemberAttributesChanged[0].Notification.ChangedAttributes.Find(LobbyMemberAttribute1Key));
-		UTEST_EQUAL("Check LobbySchemaChanged is valid.", EventCapture.LobbySchemaChanged.Num(), 1);
-		UTEST_EQUAL("Check LobbySchemaChanged is valid.", EventCapture.LobbySchemaChanged[0].Notification.Lobby, ClientData.GetPublicDataPtr());
-		UTEST_EQUAL("Check LobbyLeaderChanged is valid.", EventCapture.LobbyLeaderChanged.Num(), 1);
-		UTEST_EQUAL("Check LobbyLeaderChanged is valid.", EventCapture.LobbyLeaderChanged[0].Notification.Lobby, ClientData.GetPublicDataPtr());
-		UTEST_EQUAL("Check LobbyLeaderChanged is valid.", EventCapture.LobbyLeaderChanged[0].Notification.Leader->AccountId, User2);
-
-		// 2.4 setup
+			TOnlineResult<FLobbyClientDataPrepareServiceSnapshot> PrepareResult = ClientData.PrepareServiceSnapshot({ MoveTemp(LobbySnapshot), MoveTemp(LobbyMemberSnapshots), MoveTemp(LeaveReasons) });
+			UTEST_TRUE("Check PrepareServiceSnapshot succeeded.", PrepareResult.IsOk());
+
+			FLobbyClientDataCommitServiceSnapshot::Result CommitResult = ClientData.CommitServiceSnapshot({ &LobbyEvents });
+			UTEST_EQUAL("Check no local members left.", CommitResult.LeavingLocalMembers.Num(), 0);
+		}
+
+		{
+			// 2.6.1. Snapshot with changed lobby attribute triggers lobby attribute changed for changed attribute.
+			// 2.6.2. Snapshot with changed lobby member attribute triggers lobby member attribute changed for changed attribute.
+			UTEST_EQUAL("Check events received.", EventCapture.GetTotalNotificationsReceived(), 2);
+			UTEST_EQUAL("Check LobbyAttributesChanged is valid.", EventCapture.LobbyAttributesChanged.Num(), 1);
+			UTEST_EQUAL("Check LobbyAttributesChanged is valid.", EventCapture.LobbyAttributesChanged[0].Notification.Lobby, ClientData.GetPublicDataPtr());
+			UTEST_EQUAL("Check LobbyAttributesChanged is valid.", EventCapture.LobbyAttributesChanged[0].Notification.AddedAttributes.Num(), 0);
+			UTEST_EQUAL("Check LobbyAttributesChanged is valid.", EventCapture.LobbyAttributesChanged[0].Notification.ChangedAttributes.Num(), 0);
+			UTEST_EQUAL("Check LobbyAttributesChanged is valid.", EventCapture.LobbyAttributesChanged[0].Notification.RemovedAttributes.Num(), 1);
+			UTEST_TRUE("Check LobbyAttributesChanged is valid.", EventCapture.LobbyAttributesChanged[0].Notification.RemovedAttributes.Contains(LobbyAttributeId2));
+			UTEST_EQUAL("Check LobbyMemberAttributesChanged is valid.", EventCapture.LobbyMemberAttributesChanged.Num(), 1);
+			UTEST_EQUAL("Check LobbyMemberAttributesChanged is valid.", EventCapture.LobbyMemberAttributesChanged[0].Notification.Lobby, ClientData.GetPublicDataPtr());
+			UTEST_EQUAL("Check LobbyMemberAttributesChanged is valid.", EventCapture.LobbyMemberAttributesChanged[0].Notification.Member->AccountId, User1);
+			UTEST_EQUAL("Check LobbyMemberAttributesChanged is valid.", EventCapture.LobbyMemberAttributesChanged[0].Notification.AddedAttributes.Num(), 0);
+			UTEST_EQUAL("Check LobbyMemberAttributesChanged is valid.", EventCapture.LobbyMemberAttributesChanged[0].Notification.ChangedAttributes.Num(), 0);
+			UTEST_EQUAL("Check LobbyMemberAttributesChanged is valid.", EventCapture.LobbyMemberAttributesChanged[0].Notification.RemovedAttributes.Num(), 1);
+			UTEST_TRUE("Check LobbyMemberAttributesChanged is valid.", EventCapture.LobbyMemberAttributesChanged[0].Notification.RemovedAttributes.Contains(LobbyMemberAttributeId2));
+		}
+
+		// 2.7 setup - kick local user from lobby.
 		{
 			EventCapture.Empty();
 
-			TMap<FAccountId, TSharedRef<FClientLobbyMemberSnapshot>> LobbyMemberSnapshots;
+			TMap<FAccountId, FLobbyMemberServiceSnapshot> LobbyMemberSnapshots;
 			{
-				TSharedRef<FClientLobbyMemberSnapshot> MemberSnapshot = MakeShared<FClientLobbyMemberSnapshot>();
-				MemberSnapshot->AccountId = User2;
-				LobbyMemberSnapshots.Add(User2, MemberSnapshot);
+				FLobbyMemberServiceSnapshot& MemberSnapshot = LobbyMemberSnapshots.Add(User2);
+				MemberSnapshot.AccountId = User2;
 			}
 
-			FClientLobbySnapshot LobbySnapshot;
+			FLobbyServiceSnapshot LobbySnapshot;
 			LobbySnapshot.OwnerAccountId = User2;
-			LobbySnapshot.SchemaName = LobbySchemaName;
 			LobbySnapshot.MaxMembers = LobbyMaxMembers;
 			LobbySnapshot.JoinPolicy = LobbyJoinPolicy;
+			LobbySnapshot.SchemaServiceSnapshot.Attributes.Add(LobbyServiceAttributeId1, SchemaCompatibilityAttributeValue);
+			LobbySnapshot.SchemaServiceSnapshot.Attributes.Add(LobbyServiceAttributeId2, LobbyAttribute1Value1);
 			LobbySnapshot.Members = {User2};
 
 			TMap<FAccountId, ELobbyMemberLeaveReason> LeaveReasons;
 			LeaveReasons.Add(User1, ELobbyMemberLeaveReason::Kicked);
 
-			FApplyLobbyUpdateResult Result = ClientData.ApplyLobbyUpdateFromServiceSnapshot(MoveTemp(LobbySnapshot), MoveTemp(LobbyMemberSnapshots), MoveTemp(LeaveReasons), &LobbyEvents);
-			UTEST_EQUAL("Check local member left.", Result.LeavingLocalMembers.Num(), 1);
-			UTEST_EQUAL("Check local member left.", Result.LeavingLocalMembers[0], User1);
-		}
-
-		// 2.4.1. Snapshot with cleared lobby attribute triggers lobby attribute changed.
-		// 2.4.2. Snapshot with cleared lobby member attribute triggers lobby member attribute changed.
-		// 2.4.3. Snapshot with members leaving triggers lobby member left and lobby left.
-		UTEST_EQUAL("Check events received.", EventCapture.GetTotalNotificationsReceived(), 5);
-		UTEST_EQUAL("Check LobbyAttributesChanged is valid.", EventCapture.LobbyAttributesChanged.Num(), 1);
-		UTEST_EQUAL("Check LobbyAttributesChanged is valid.", EventCapture.LobbyAttributesChanged[0].Notification.Lobby, ClientData.GetPublicDataPtr());
-		UTEST_EQUAL("Check LobbyAttributesChanged is valid.", EventCapture.LobbyAttributesChanged[0].Notification.ChangedAttributes.Num(), 1);
-		UTEST_NOT_NULL("Check LobbyAttributesChanged is valid.", EventCapture.LobbyAttributesChanged[0].Notification.ChangedAttributes.Find(LobbyAttribute1Key));
-		UTEST_EQUAL("Check LobbyMemberAttributesChanged is valid.", EventCapture.LobbyMemberAttributesChanged.Num(), 1);
-		UTEST_EQUAL("Check LobbyMemberAttributesChanged is valid.", EventCapture.LobbyMemberAttributesChanged[0].Notification.Lobby, ClientData.GetPublicDataPtr());
-		UTEST_EQUAL("Check LobbyMemberAttributesChanged is valid.", EventCapture.LobbyMemberAttributesChanged[0].Notification.Member->AccountId, User2);
-		UTEST_EQUAL("Check LobbyMemberAttributesChanged is valid.", EventCapture.LobbyMemberAttributesChanged[0].Notification.ChangedAttributes.Num(), 1);
-		UTEST_NOT_NULL("Check LobbyMemberAttributesChanged is valid.", EventCapture.LobbyMemberAttributesChanged[0].Notification.ChangedAttributes.Find(LobbyMemberAttribute1Key));
-		UTEST_EQUAL("Check LobbyMemberLeft valid.", EventCapture.LobbyMemberLeft.Num(), 2);
-		FLobbyEventCapture::TNotificationInfo<FLobbyMemberLeft>* User1LeaveNotification = EventCapture.LobbyMemberLeft[0].Notification.Member->AccountId == User1 ? &EventCapture.LobbyMemberLeft[0] : &EventCapture.LobbyMemberLeft[1];
-		FLobbyEventCapture::TNotificationInfo<FLobbyMemberLeft>* User2LeaveNotification = EventCapture.LobbyMemberLeft[0].Notification.Member->AccountId == User2 ? &EventCapture.LobbyMemberLeft[0] : &EventCapture.LobbyMemberLeft[1];
-		UTEST_EQUAL("Check LobbyMemberLeft valid.", User1LeaveNotification->Notification.Lobby, ClientData.GetPublicDataPtr());
-		UTEST_EQUAL("Check LobbyMemberLeft valid.", User1LeaveNotification->Notification.Member->AccountId, User1);
-		UTEST_EQUAL("Check LobbyMemberLeft valid.", User1LeaveNotification->Notification.Reason, ELobbyMemberLeaveReason::Kicked);
-		UTEST_EQUAL("Check LobbyMemberLeft valid.", User2LeaveNotification->Notification.Lobby, ClientData.GetPublicDataPtr());
-		UTEST_EQUAL("Check LobbyMemberLeft valid.", User2LeaveNotification->Notification.Member->AccountId, User2);
-		UTEST_EQUAL("Check LobbyMemberLeft valid.", User2LeaveNotification->Notification.Reason, ELobbyMemberLeaveReason::Left);
-		UTEST_EQUAL("Check LobbyLeft valid.", EventCapture.LobbyLeft.Num(), 1);
-		UTEST_EQUAL("Check LobbyLeft valid.", EventCapture.LobbyLeft[0].Notification.Lobby, ClientData.GetPublicDataPtr());
-		UTEST_EQUAL("Check LobbyLeft ordering.", EventCapture.LobbyLeft[0].GlobalIndex, EventCapture.GetTotalNotificationsReceived() - 1);
-		FLobbyEventCapture::TNotificationInfo<FLobbyMemberLeft>* LastMemberLeaveNotification = User1LeaveNotification->GlobalIndex > User2LeaveNotification->GlobalIndex ? User1LeaveNotification : User2LeaveNotification;
-		UTEST_EQUAL("Check LobbyMemberLeft ordering.", LastMemberLeaveNotification->GlobalIndex, EventCapture.LobbyLeft[0].GlobalIndex - 1);
+			TOnlineResult<FLobbyClientDataPrepareServiceSnapshot> PrepareResult = ClientData.PrepareServiceSnapshot({ MoveTemp(LobbySnapshot), MoveTemp(LobbyMemberSnapshots), MoveTemp(LeaveReasons) });
+			UTEST_TRUE("Check PrepareServiceSnapshot succeeded.", PrepareResult.IsOk());
+
+			FLobbyClientDataCommitServiceSnapshot::Result CommitResult = ClientData.CommitServiceSnapshot({ &LobbyEvents });
+			UTEST_EQUAL("Check local member left.", CommitResult.LeavingLocalMembers.Num(), 1);
+			UTEST_EQUAL("Check local member left.", CommitResult.LeavingLocalMembers[0], User1);
+		}
+
+		{
+			// 2.7.1. Snapshot with members leaving triggers lobby member left and lobby left.
+			UTEST_EQUAL("Check events received.", EventCapture.GetTotalNotificationsReceived(), 3);
+			UTEST_EQUAL("Check LobbyMemberLeft valid.", EventCapture.LobbyMemberLeft.Num(), 2);
+			FLobbyEventCapture::TNotificationInfo<FLobbyMemberLeft>* User1LeaveNotification = EventCapture.LobbyMemberLeft[0].Notification.Member->AccountId == User1 ? &EventCapture.LobbyMemberLeft[0] : &EventCapture.LobbyMemberLeft[1];
+			FLobbyEventCapture::TNotificationInfo<FLobbyMemberLeft>* User2LeaveNotification = EventCapture.LobbyMemberLeft[0].Notification.Member->AccountId == User2 ? &EventCapture.LobbyMemberLeft[0] : &EventCapture.LobbyMemberLeft[1];
+			UTEST_EQUAL("Check LobbyMemberLeft valid.", User1LeaveNotification->Notification.Lobby, ClientData.GetPublicDataPtr());
+			UTEST_EQUAL("Check LobbyMemberLeft valid.", User1LeaveNotification->Notification.Member->AccountId, User1);
+			UTEST_EQUAL("Check LobbyMemberLeft valid.", User1LeaveNotification->Notification.Reason, ELobbyMemberLeaveReason::Kicked);
+			UTEST_EQUAL("Check LobbyMemberLeft valid.", User2LeaveNotification->Notification.Lobby, ClientData.GetPublicDataPtr());
+			UTEST_EQUAL("Check LobbyMemberLeft valid.", User2LeaveNotification->Notification.Member->AccountId, User2);
+			UTEST_EQUAL("Check LobbyMemberLeft valid.", User2LeaveNotification->Notification.Reason, ELobbyMemberLeaveReason::Left);
+			UTEST_EQUAL("Check LobbyLeft valid.", EventCapture.LobbyLeft.Num(), 1);
+			UTEST_EQUAL("Check LobbyLeft valid.", EventCapture.LobbyLeft[0].Notification.Lobby, ClientData.GetPublicDataPtr());
+			UTEST_EQUAL("Check LobbyLeft ordering.", EventCapture.LobbyLeft[0].GlobalIndex, EventCapture.GetTotalNotificationsReceived() - 1);
+			FLobbyEventCapture::TNotificationInfo<FLobbyMemberLeft>* LastMemberLeaveNotification = User1LeaveNotification->GlobalIndex > User2LeaveNotification->GlobalIndex ? User1LeaveNotification : User2LeaveNotification;
+			UTEST_EQUAL("Check LobbyMemberLeft ordering.", LastMemberLeaveNotification->GlobalIndex, EventCapture.LobbyLeft[0].GlobalIndex - 1);
+		}
 	}
 
 	// 3. Snapshot supplemental
@@ -455,142 +791,151 @@ bool FClientLobbyDataTest::RunTest(const FString& Parameters)
 		FLobbyEvents LobbyEvents;
 		FLobbyEventCapture EventCapture(LobbyEvents);
 		FLobbyId LobbyId(EOnlineServices::Null, 1);
-		FClientLobbyData ClientData(LobbyId);
+		FLobbyClientData ClientData(LobbyId, SchemaRegistry);
 
 		FAccountId User1(EOnlineServices::Null, 1);
 		FAccountId User2(EOnlineServices::Null, 2);
-		FName LobbySchemaName = TEXT("SchemaName");
 		int32 LobbyMaxMembers = 5;
 		ELobbyJoinPolicy LobbyJoinPolicy = ELobbyJoinPolicy::PublicAdvertised;
 
 		// 3.1 setup - initialize lobby data with initial snapshot.
 		{
-			TMap<FAccountId, TSharedRef<FClientLobbyMemberSnapshot>> LobbyMemberSnapshots;
+			TMap<FAccountId, FLobbyMemberServiceSnapshot> LobbyMemberSnapshots;
 			{
-				TSharedRef<FClientLobbyMemberSnapshot> MemberSnapshot = MakeShared<FClientLobbyMemberSnapshot>();
-				MemberSnapshot->AccountId = User1;
-				LobbyMemberSnapshots.Add(User1, MemberSnapshot);
+				FLobbyMemberServiceSnapshot& MemberSnapshot = LobbyMemberSnapshots.Add(User1);
+				MemberSnapshot.AccountId = User1;
 			}
 			{
-				TSharedRef<FClientLobbyMemberSnapshot> MemberSnapshot = MakeShared<FClientLobbyMemberSnapshot>();
-				MemberSnapshot->AccountId = User2;
-				LobbyMemberSnapshots.Add(User2, MemberSnapshot);
+				FLobbyMemberServiceSnapshot& MemberSnapshot = LobbyMemberSnapshots.Add(User2);
+				MemberSnapshot.AccountId = User2;
 			}
 
-			FClientLobbySnapshot LobbySnapshot;
+			FLobbyServiceSnapshot LobbySnapshot;
 			LobbySnapshot.OwnerAccountId = User1;
-			LobbySnapshot.SchemaName = LobbySchemaName;
 			LobbySnapshot.MaxMembers = LobbyMaxMembers;
 			LobbySnapshot.JoinPolicy = LobbyJoinPolicy;
+			LobbySnapshot.SchemaServiceSnapshot.Attributes.Add(LobbyServiceAttributeId1, SchemaCompatibilityAttributeValue);
 			LobbySnapshot.Members = {User1, User2};
 
 			TMap<FAccountId, ELobbyMemberLeaveReason> LeaveReasons;
-			FApplyLobbyUpdateResult Result = ClientData.ApplyLobbyUpdateFromServiceSnapshot(MoveTemp(LobbySnapshot), MoveTemp(LobbyMemberSnapshots), MoveTemp(LeaveReasons), &LobbyEvents);
-			UTEST_EQUAL("Check no local members left.", Result.LeavingLocalMembers.Num(), 0);
+			TOnlineResult<FLobbyClientDataPrepareServiceSnapshot> PrepareResult = ClientData.PrepareServiceSnapshot({ MoveTemp(LobbySnapshot), MoveTemp(LobbyMemberSnapshots), MoveTemp(LeaveReasons) });
+			UTEST_TRUE("Check PrepareServiceSnapshot succeeded.", PrepareResult.IsOk());
+
+			FLobbyClientDataCommitServiceSnapshot::Result CommitResult = ClientData.CommitServiceSnapshot({ &LobbyEvents });
+			UTEST_EQUAL("Check no local members left.", CommitResult.LeavingLocalMembers.Num(), 0);
 		}
 
 		// 3.1 setup continued. Add local member to enable notifications.
 		{
-			FClientLobbyDataChanges LobbyChanges;
-			LobbyChanges.MutatedMembers.Add(User1, MakeShared<FClientLobbyMemberDataChanges>());
-			FApplyLobbyUpdateResult Result = ClientData.ApplyLobbyUpdateFromLocalChanges(MoveTemp(LobbyChanges), LobbyEvents);
-			UTEST_EQUAL("Check no local members left.", Result.LeavingLocalMembers.Num(), 0);
+			FLobbyClientChanges LobbyChanges;
+			LobbyChanges.MemberAttributes = FLobbyClientMemberAttributeChanges();
+			TOnlineResult<FLobbyClientDataPrepareClientChanges> PrepareResult = ClientData.PrepareClientChanges({ User1, MoveTemp(LobbyChanges) });
+			UTEST_TRUE("Check PrepareClientChanges succeeded.", PrepareResult.IsOk());
+
+			FLobbyClientDataCommitClientChanges::Result CommitResult = ClientData.CommitClientChanges({ &LobbyEvents });
+			UTEST_EQUAL("Check no local members left.", CommitResult.LeavingLocalMembers.Num(), 0);
 		}
 
-		// 3.1 setup continued. Change snapshot for the test.
+		// 3.1 setup continued. Change snapshot for the test. User 2 leaves with no explicit reason.
 		{
 			EventCapture.Empty();
 
-			TMap<FAccountId, TSharedRef<FClientLobbyMemberSnapshot>> LobbyMemberSnapshots;
+			TMap<FAccountId, FLobbyMemberServiceSnapshot> LobbyMemberSnapshots;
 			{
-				TSharedRef<FClientLobbyMemberSnapshot> MemberSnapshot = MakeShared<FClientLobbyMemberSnapshot>();
-				MemberSnapshot->AccountId = User1;
-				LobbyMemberSnapshots.Add(User1, MemberSnapshot);
+				FLobbyMemberServiceSnapshot& MemberSnapshot = LobbyMemberSnapshots.Add(User1);
+				MemberSnapshot.AccountId = User1;
 			}
 
-			FClientLobbySnapshot LobbySnapshot;
+			FLobbyServiceSnapshot LobbySnapshot;
 			LobbySnapshot.OwnerAccountId = User1;
-			LobbySnapshot.SchemaName = LobbySchemaName;
 			LobbySnapshot.MaxMembers = LobbyMaxMembers;
 			LobbySnapshot.JoinPolicy = LobbyJoinPolicy;
+			LobbySnapshot.SchemaServiceSnapshot.Attributes.Add(LobbyServiceAttributeId1, SchemaCompatibilityAttributeValue);
 			LobbySnapshot.Members = {User1};
 
 			TMap<FAccountId, ELobbyMemberLeaveReason> LeaveReasons;
-			FApplyLobbyUpdateResult Result = ClientData.ApplyLobbyUpdateFromServiceSnapshot(MoveTemp(LobbySnapshot), MoveTemp(LobbyMemberSnapshots), MoveTemp(LeaveReasons), &LobbyEvents);
-			UTEST_EQUAL("Check no local members left.", Result.LeavingLocalMembers.Num(), 0);
+			TOnlineResult<FLobbyClientDataPrepareServiceSnapshot> PrepareResult = ClientData.PrepareServiceSnapshot({ MoveTemp(LobbySnapshot), MoveTemp(LobbyMemberSnapshots), MoveTemp(LeaveReasons) });
+			UTEST_TRUE("Check PrepareServiceSnapshot succeeded.", PrepareResult.IsOk());
+
+			FLobbyClientDataCommitServiceSnapshot::Result CommitResult = ClientData.CommitServiceSnapshot({ &LobbyEvents });
+			UTEST_EQUAL("Check no local members left.", CommitResult.LeavingLocalMembers.Num(), 0);
 		}
 
-		// 3.1. Test that Member left is fired even when a leave reason is not specified.
-		UTEST_EQUAL("Check events received.", EventCapture.GetTotalNotificationsReceived(), 1);
-		UTEST_EQUAL("Check LobbyMemberLeft valid.", EventCapture.LobbyMemberLeft.Num(), 1);
-		UTEST_EQUAL("Check LobbyMemberLeft valid.", EventCapture.LobbyMemberLeft[0].Notification.Lobby, ClientData.GetPublicDataPtr());
-		UTEST_EQUAL("Check LobbyMemberLeft valid.", EventCapture.LobbyMemberLeft[0].Notification.Member->AccountId, User2);
-		UTEST_EQUAL("Check LobbyMemberLeft valid.", EventCapture.LobbyMemberLeft[0].Notification.Reason, ELobbyMemberLeaveReason::Disconnected);
+		{
+			// 3.1. Test that Member left is fired even when a leave reason is not specified.
+			UTEST_EQUAL("Check events received.", EventCapture.GetTotalNotificationsReceived(), 1);
+			UTEST_EQUAL("Check LobbyMemberLeft valid.", EventCapture.LobbyMemberLeft.Num(), 1);
+			UTEST_EQUAL("Check LobbyMemberLeft valid.", EventCapture.LobbyMemberLeft[0].Notification.Lobby, ClientData.GetPublicDataPtr());
+			UTEST_EQUAL("Check LobbyMemberLeft valid.", EventCapture.LobbyMemberLeft[0].Notification.Member->AccountId, User2);
+			UTEST_EQUAL("Check LobbyMemberLeft valid.", EventCapture.LobbyMemberLeft[0].Notification.Reason, ELobbyMemberLeaveReason::Disconnected);
+		}
 
 		// 3.2 setup
 		{
 			EventCapture.Empty();
 
-			TMap<FAccountId, TSharedRef<FClientLobbyMemberSnapshot>> LobbyMemberSnapshots;
+			TMap<FAccountId, FLobbyMemberServiceSnapshot> LobbyMemberSnapshots;
 			{
-				TSharedRef<FClientLobbyMemberSnapshot> MemberSnapshot = MakeShared<FClientLobbyMemberSnapshot>();
-				MemberSnapshot->AccountId = User1;
-				LobbyMemberSnapshots.Add(User1, MemberSnapshot);
+				FLobbyMemberServiceSnapshot& MemberSnapshot = LobbyMemberSnapshots.Add(User1);
+				MemberSnapshot.AccountId = User1;
 			}
 			{
-				TSharedRef<FClientLobbyMemberSnapshot> MemberSnapshot = MakeShared<FClientLobbyMemberSnapshot>();
-				MemberSnapshot->AccountId = User2;
-				LobbyMemberSnapshots.Add(User2, MemberSnapshot);
+				FLobbyMemberServiceSnapshot& MemberSnapshot = LobbyMemberSnapshots.Add(User2);
+				MemberSnapshot.AccountId = User2;
 			}
 
-			FClientLobbySnapshot LobbySnapshot;
+			FLobbyServiceSnapshot LobbySnapshot;
 			LobbySnapshot.OwnerAccountId = User1;
-			LobbySnapshot.SchemaName = LobbySchemaName;
 			LobbySnapshot.MaxMembers = LobbyMaxMembers;
 			LobbySnapshot.JoinPolicy = LobbyJoinPolicy;
+			LobbySnapshot.SchemaServiceSnapshot.Attributes.Add(LobbyServiceAttributeId1, SchemaCompatibilityAttributeValue);
 			LobbySnapshot.Members = {User1};
 
 			TMap<FAccountId, ELobbyMemberLeaveReason> LeaveReasons;
-			FApplyLobbyUpdateResult Result = ClientData.ApplyLobbyUpdateFromServiceSnapshot(MoveTemp(LobbySnapshot), MoveTemp(LobbyMemberSnapshots), MoveTemp(LeaveReasons), &LobbyEvents);
-			UTEST_EQUAL("Check no local members left.", Result.LeavingLocalMembers.Num(), 0);
+			TOnlineResult<FLobbyClientDataPrepareServiceSnapshot> PrepareResult = ClientData.PrepareServiceSnapshot({ MoveTemp(LobbySnapshot), MoveTemp(LobbyMemberSnapshots), MoveTemp(LeaveReasons) });
+			UTEST_TRUE("Check PrepareServiceSnapshot succeeded.", PrepareResult.IsOk());
+
+			FLobbyClientDataCommitServiceSnapshot::Result CommitResult = ClientData.CommitServiceSnapshot({ &LobbyEvents });
+			UTEST_EQUAL("Check no local members left.", CommitResult.LeavingLocalMembers.Num(), 0);
 		}
 
-		// 3.2. Applying member snapshots for an account id not in the lobby members list does not fire lobby member joined.
-		UTEST_EQUAL("Check no events received.", EventCapture.GetTotalNotificationsReceived(), 0);
+		{
+			// 3.2. Applying member snapshots for an account id not in the lobby members list does not fire lobby member joined.
+			UTEST_EQUAL("Check no events received.", EventCapture.GetTotalNotificationsReceived(), 0);
+		}
 
 		// 3.3 setup
 		{
 			EventCapture.Empty();
 
-			TMap<FAccountId, TSharedRef<FClientLobbyMemberSnapshot>> LobbyMemberSnapshots;
+			TMap<FAccountId, FLobbyMemberServiceSnapshot> LobbyMemberSnapshots;
 			{
-				TSharedRef<FClientLobbyMemberSnapshot> MemberSnapshot = MakeShared<FClientLobbyMemberSnapshot>();
-				MemberSnapshot->AccountId = User1;
-				LobbyMemberSnapshots.Add(User1, MemberSnapshot);
-			}
-			{
-				TSharedRef<FClientLobbyMemberSnapshot> MemberSnapshot = MakeShared<FClientLobbyMemberSnapshot>();
-				MemberSnapshot->AccountId = User2;
-				LobbyMemberSnapshots.Add(User2, MemberSnapshot);
+				FLobbyMemberServiceSnapshot& MemberSnapshot = LobbyMemberSnapshots.Add(User2);
+				MemberSnapshot.AccountId = User2;
 			}
 
-			FClientLobbySnapshot LobbySnapshot;
+			FLobbyServiceSnapshot LobbySnapshot;
 			LobbySnapshot.OwnerAccountId = User1;
-			LobbySnapshot.SchemaName = LobbySchemaName;
 			LobbySnapshot.MaxMembers = LobbyMaxMembers;
 			LobbySnapshot.JoinPolicy = LobbyJoinPolicy;
+			LobbySnapshot.SchemaServiceSnapshot.Attributes.Add(LobbyServiceAttributeId1, SchemaCompatibilityAttributeValue);
 			LobbySnapshot.Members = {User1, User2};
 
 			TMap<FAccountId, ELobbyMemberLeaveReason> LeaveReasons;
-			FApplyLobbyUpdateResult Result = ClientData.ApplyLobbyUpdateFromServiceSnapshot(MoveTemp(LobbySnapshot), MoveTemp(LobbyMemberSnapshots), MoveTemp(LeaveReasons), &LobbyEvents);
-			UTEST_EQUAL("Check no local members left.", Result.LeavingLocalMembers.Num(), 0);
+			TOnlineResult<FLobbyClientDataPrepareServiceSnapshot> PrepareResult = ClientData.PrepareServiceSnapshot({ MoveTemp(LobbySnapshot), MoveTemp(LobbyMemberSnapshots), MoveTemp(LeaveReasons) });
+			UTEST_TRUE("Check PrepareServiceSnapshot succeeded.", PrepareResult.IsOk());
+
+			FLobbyClientDataCommitServiceSnapshot::Result CommitResult = ClientData.CommitServiceSnapshot({ &LobbyEvents });
+			UTEST_EQUAL("Check no local members left.", CommitResult.LeavingLocalMembers.Num(), 0);
 		}
 
-		// 3.3. Check lobby member joined notification received for remote member after local member joined.
-		UTEST_EQUAL("Check events received.", EventCapture.GetTotalNotificationsReceived(), 1);
-		UTEST_EQUAL("Check LobbyMemberJoined valid.", EventCapture.LobbyMemberJoined.Num(), 1);
-		UTEST_EQUAL("Check LobbyMemberJoined valid.", EventCapture.LobbyMemberJoined[0].Notification.Lobby, ClientData.GetPublicDataPtr());
-		UTEST_EQUAL("Check LobbyMemberJoined valid.", EventCapture.LobbyMemberJoined[0].Notification.Member->AccountId, User2);
+		{
+			// 3.3.1 Check lobby member joined notification received for remote member after local member joined.
+			UTEST_EQUAL("Check events received.", EventCapture.GetTotalNotificationsReceived(), 1);
+			UTEST_EQUAL("Check LobbyMemberJoined valid.", EventCapture.LobbyMemberJoined.Num(), 1);
+			UTEST_EQUAL("Check LobbyMemberJoined valid.", EventCapture.LobbyMemberJoined[0].Notification.Lobby, ClientData.GetPublicDataPtr());
+			UTEST_EQUAL("Check LobbyMemberJoined valid.", EventCapture.LobbyMemberJoined[0].Notification.Member->AccountId, User2);
+		}
 	}
 
 	// 4. Applying lobby changes produces expected notifications.
@@ -598,178 +943,264 @@ bool FClientLobbyDataTest::RunTest(const FString& Parameters)
 		FLobbyEvents LobbyEvents;
 		FLobbyEventCapture EventCapture(LobbyEvents);
 		FLobbyId LobbyId(EOnlineServices::Null, 1);
-		FClientLobbyData ClientData(LobbyId);
+		FLobbyClientData ClientData(LobbyId, SchemaRegistry);
 
 		FAccountId User1(EOnlineServices::Null, 1);
 		FAccountId User2(EOnlineServices::Null, 2);
-		FName LobbySchemaName = TEXT("SchemaName");
 		int32 LobbyMaxMembers = 5;
 		ELobbyJoinPolicy LobbyJoinPolicy = ELobbyJoinPolicy::PublicAdvertised;
-		FLobbyAttributeId LobbyAttribute1Key = TEXT("LobbyAttribute1");
-		FLobbyVariant LobbyAttribute1Value;
-		LobbyAttribute1Value.Set(TEXT("Attribute1"));
-		FLobbyAttributeId LobbyAttribute2Key = TEXT("LobbyAttribute2");
-		FLobbyVariant LobbyAttribute2Value;
-		LobbyAttribute2Value.Set(TEXT("Attribute2"));
-		FLobbyAttributeId LobbyMemberAttribute1Key = TEXT("LobbyMemberAttribute1");
-		FLobbyVariant LobbyMemberAttribute1Value;
-		LobbyMemberAttribute1Value.Set(TEXT("MemberAttribute1"));
-		FLobbyAttributeId LobbyMemberAttribute2Key = TEXT("LobbyMemberAttribute2");
-		FLobbyVariant LobbyMemberAttribute2Value;
-		LobbyMemberAttribute2Value.Set(TEXT("MemberAttribute2"));
 
 		// 4.1 setup. Local member in snapshot with remote member to simulate lobby join.
 		{
-			TMap<FAccountId, TSharedRef<FClientLobbyMemberSnapshot>> LobbyMemberSnapshots;
+			TMap<FAccountId, FLobbyMemberServiceSnapshot> LobbyMemberSnapshots;
 			{
-				TSharedRef<FClientLobbyMemberSnapshot> MemberSnapshot = MakeShared<FClientLobbyMemberSnapshot>();
-				MemberSnapshot->AccountId = User1;
-				MemberSnapshot->Attributes.Add(LobbyMemberAttribute1Key, LobbyMemberAttribute1Value);
-				LobbyMemberSnapshots.Add(User1, MemberSnapshot);
+				FLobbyMemberServiceSnapshot& MemberSnapshot = LobbyMemberSnapshots.Add(User1);
+				MemberSnapshot.AccountId = User1;
+				MemberSnapshot.SchemaServiceSnapshot.Attributes.Add(LobbyMemberServiceAttributeId1, LobbyMemberAttribute1Value1);
 			}
 			{
-				TSharedRef<FClientLobbyMemberSnapshot> MemberSnapshot = MakeShared<FClientLobbyMemberSnapshot>();
-				MemberSnapshot->AccountId = User2;
-				LobbyMemberSnapshots.Add(User2, MemberSnapshot);
+				FLobbyMemberServiceSnapshot& MemberSnapshot = LobbyMemberSnapshots.Add(User2);
+				MemberSnapshot.AccountId = User2;
 			}
 
-			FClientLobbySnapshot LobbySnapshot;
+			FLobbyServiceSnapshot LobbySnapshot;
 			LobbySnapshot.OwnerAccountId = User1;
-			LobbySnapshot.SchemaName = LobbySchemaName;
 			LobbySnapshot.MaxMembers = LobbyMaxMembers;
 			LobbySnapshot.JoinPolicy = LobbyJoinPolicy;
-			LobbySnapshot.Attributes.Add(LobbyAttribute1Key, LobbyAttribute1Value);
+			LobbySnapshot.SchemaServiceSnapshot.Attributes.Add(LobbyServiceAttributeId1, SchemaCompatibilityAttributeValue);
+			LobbySnapshot.SchemaServiceSnapshot.Attributes.Add(LobbyServiceAttributeId2, LobbyAttribute1Value1);
 			LobbySnapshot.Members = {User1, User2};
 
 			TMap<FAccountId, ELobbyMemberLeaveReason> LeaveReasons;
-			FApplyLobbyUpdateResult Result = ClientData.ApplyLobbyUpdateFromServiceSnapshot(MoveTemp(LobbySnapshot), MoveTemp(LobbyMemberSnapshots), MoveTemp(LeaveReasons), &LobbyEvents);
-			UTEST_EQUAL("Check no local members left.", Result.LeavingLocalMembers.Num(), 0);
+			TOnlineResult<FLobbyClientDataPrepareServiceSnapshot> PrepareResult = ClientData.PrepareServiceSnapshot({ MoveTemp(LobbySnapshot), MoveTemp(LobbyMemberSnapshots), MoveTemp(LeaveReasons) });
+			UTEST_TRUE("Check PrepareServiceSnapshot succeeded.", PrepareResult.IsOk());
+
+			FLobbyClientDataCommitServiceSnapshot::Result CommitResult = ClientData.CommitServiceSnapshot({ &LobbyEvents });
+			UTEST_EQUAL("Check no local members left.", CommitResult.LeavingLocalMembers.Num(), 0);
 		}
 
 		UTEST_EQUAL("Check no events received.", EventCapture.GetTotalNotificationsReceived(), 0);
 
-		// 4.1 setup continued.
+		// 4.1 Attempt to update lobby data when user is not the owner fails.
 		{
-			LobbyAttribute1Value.Set(TEXT("MutatedAttribute1"));
-			LobbySchemaName = TEXT("Beans");
-
-			FClientLobbyDataChanges LobbyChanges;
-			LobbyChanges.MutatedAttributes.Add(LobbyAttribute1Key, LobbyAttribute1Value);
-			LobbyChanges.LobbySchema = LobbySchemaName;
+			FLobbyClientChanges LobbyChanges;
+			LobbyChanges.Attributes = { { { LobbyAttributeId1, LobbyAttribute1Value2 } }, {} };
 			LobbyChanges.OwnerAccountId = User2;
-			FApplyLobbyUpdateResult Result = ClientData.ApplyLobbyUpdateFromLocalChanges(MoveTemp(LobbyChanges), LobbyEvents);
-			UTEST_EQUAL("Check no local members left.", Result.LeavingLocalMembers.Num(), 0);
+			TOnlineResult<FLobbyClientDataPrepareClientChanges> PrepareResult = ClientData.PrepareClientChanges({ User2, MoveTemp(LobbyChanges) });
+			UTEST_FALSE("Check PrepareClientChanges failed as expected.", PrepareResult.IsOk());
 		}
 
-		// 4.1.1. Local changes with changed lobby attribute does not trigger lobby attribute changed.
-		// 4.1.2. Local changes with changed lobby schema does not trigger lobby schema changed.
-		// 4.1.3. Local changes with changed lobby leader does not trigger lobby leader changed.
-		UTEST_EQUAL("Check no events received.", EventCapture.GetTotalNotificationsReceived(), 0);
+		// 4.2 setup. Join local User1 to the lobby to enable notifications.
+		{
+			FLobbyClientChanges LobbyChanges;
+			LobbyChanges.MemberAttributes = FLobbyClientMemberAttributeChanges();
+			TOnlineResult<FLobbyClientDataPrepareClientChanges> PrepareResult = ClientData.PrepareClientChanges({ User1, MoveTemp(LobbyChanges) });
+			UTEST_TRUE("Check PrepareClientChanges succeeded.", PrepareResult.IsOk());
 
-		// 4.2 setup.
+			FLobbyClientDataCommitClientChanges::Result CommitResult = ClientData.CommitClientChanges({ &LobbyEvents });
+			UTEST_EQUAL("Check no local members left.", CommitResult.LeavingLocalMembers.Num(), 0);
+		}
+
+		{
+			// 4.2.1. Local members added to lobby through local change triggers lobby joined and lobby member joined. On member joined triggered for remote member.
+			UTEST_EQUAL("Check events received.", EventCapture.GetTotalNotificationsReceived(), 3);
+			UTEST_EQUAL("Check LobbyJoined valid.", EventCapture.LobbyJoined.Num(), 1);
+			UTEST_EQUAL("Check LobbyJoined valid.", EventCapture.LobbyJoined[0].GlobalIndex, 0);
+			UTEST_EQUAL("Check LobbyMemberJoined valid.", EventCapture.LobbyJoined[0].Notification.Lobby->Attributes.Num(), 1);
+			const FSchemaVariant* LobbyAttribute1 = EventCapture.LobbyJoined[0].Notification.Lobby->Attributes.Find(LobbyAttributeId1);
+			UTEST_NOT_NULL("Check LobbyMemberJoined valid.", LobbyAttribute1);
+			UTEST_EQUAL("Check LobbyMemberJoined valid.", *LobbyAttribute1, LobbyAttribute1Value1);
+			UTEST_EQUAL("Check LobbyMemberJoined valid.", EventCapture.LobbyMemberJoined.Num(), 2);
+			FLobbyEventCapture::TNotificationInfo<FLobbyMemberJoined>* User1JoinNotification = EventCapture.LobbyMemberJoined[0].Notification.Member->AccountId == User1 ? &EventCapture.LobbyMemberJoined[0] : &EventCapture.LobbyMemberJoined[1];
+			FLobbyEventCapture::TNotificationInfo<FLobbyMemberJoined>* User2JoinNotification = EventCapture.LobbyMemberJoined[0].Notification.Member->AccountId == User2 ? &EventCapture.LobbyMemberJoined[0] : &EventCapture.LobbyMemberJoined[1];
+			UTEST_EQUAL("Check LobbyMemberJoined valid.", User1JoinNotification->Notification.Lobby, ClientData.GetPublicDataPtr());
+			UTEST_EQUAL("Check LobbyMemberJoined valid.", User1JoinNotification->Notification.Member->AccountId, User1);
+			UTEST_EQUAL("Check LobbyMemberJoined valid.", User1JoinNotification->Notification.Member->Attributes.Num(), 1);
+			const FSchemaVariant* LobbyMemberAttribute1 = User1JoinNotification->Notification.Member->Attributes.Find(LobbyMemberAttributeId1);
+			UTEST_NOT_NULL("Check LobbyMemberJoined valid.", LobbyMemberAttribute1);
+			UTEST_EQUAL("Check LobbyMemberJoined valid.", *LobbyMemberAttribute1, LobbyMemberAttribute1Value1);
+			UTEST_EQUAL("Check LobbyMemberJoined valid.", User2JoinNotification->Notification.Lobby, ClientData.GetPublicDataPtr());
+			UTEST_EQUAL("Check LobbyMemberJoined valid.", User2JoinNotification->Notification.Member->AccountId, User2);
+		}
+
+		// 4.3 setup - add a lobby and lobby member attribute
+		{
+			EventCapture.Empty();
+
+			FLobbyClientChanges LobbyChanges;
+			LobbyChanges.Attributes = { {{ LobbyAttributeId2, LobbyAttribute2Value1 }}, {} };
+			LobbyChanges.MemberAttributes = { {{ LobbyMemberAttributeId2, LobbyMemberAttribute2Value1 }}, {} };
+			TOnlineResult<FLobbyClientDataPrepareClientChanges> PrepareResult = ClientData.PrepareClientChanges({ User1, MoveTemp(LobbyChanges) });
+			UTEST_TRUE("Check PrepareClientChanges succeeded.", PrepareResult.IsOk());
+
+			FLobbyClientDataCommitClientChanges::Result CommitResult = ClientData.CommitClientChanges({ &LobbyEvents });
+			UTEST_EQUAL("Check no local members left.", CommitResult.LeavingLocalMembers.Num(), 0);
+		}
+
+		{
+			// 4.3.1. Local changes with changed lobby attribute triggers lobby attribute changed for adding attribute.
+			// 4.3.2. Local changes with changed lobby member attribute triggers lobby member attribute changed for adding attribute.
+			UTEST_EQUAL("Check events received.", EventCapture.GetTotalNotificationsReceived(), 2);
+			UTEST_EQUAL("Check LobbyAttributesChanged is valid.", EventCapture.LobbyAttributesChanged.Num(), 1);
+			UTEST_EQUAL("Check LobbyAttributesChanged is valid.", EventCapture.LobbyAttributesChanged[0].Notification.Lobby, ClientData.GetPublicDataPtr());
+			UTEST_EQUAL("Check LobbyAttributesChanged is valid.", EventCapture.LobbyAttributesChanged[0].Notification.AddedAttributes.Num(), 1);
+			UTEST_EQUAL("Check LobbyAttributesChanged is valid.", EventCapture.LobbyAttributesChanged[0].Notification.ChangedAttributes.Num(), 0);
+			UTEST_EQUAL("Check LobbyAttributesChanged is valid.", EventCapture.LobbyAttributesChanged[0].Notification.RemovedAttributes.Num(), 0);
+			const FSchemaVariant* LobbyAttribute2 = EventCapture.LobbyAttributesChanged[0].Notification.AddedAttributes.Find(LobbyAttributeId2);
+			UTEST_NOT_NULL("Check LobbyAttributesChanged is valid.", LobbyAttribute2);
+			UTEST_EQUAL("Check LobbyAttributesChanged is valid.", *LobbyAttribute2, LobbyAttribute2Value1);
+			UTEST_EQUAL("Check LobbyMemberAttributesChanged is valid.", EventCapture.LobbyMemberAttributesChanged.Num(), 1);
+			UTEST_EQUAL("Check LobbyMemberAttributesChanged is valid.", EventCapture.LobbyMemberAttributesChanged[0].Notification.Lobby, ClientData.GetPublicDataPtr());
+			UTEST_EQUAL("Check LobbyMemberAttributesChanged is valid.", EventCapture.LobbyMemberAttributesChanged[0].Notification.Member->AccountId, User1);
+			UTEST_EQUAL("Check LobbyMemberAttributesChanged is valid.", EventCapture.LobbyMemberAttributesChanged[0].Notification.AddedAttributes.Num(), 1);
+			UTEST_EQUAL("Check LobbyMemberAttributesChanged is valid.", EventCapture.LobbyMemberAttributesChanged[0].Notification.ChangedAttributes.Num(), 0);
+			UTEST_EQUAL("Check LobbyMemberAttributesChanged is valid.", EventCapture.LobbyMemberAttributesChanged[0].Notification.RemovedAttributes.Num(), 0);
+			const FSchemaVariant* LobbyMemberAttribute2 = EventCapture.LobbyMemberAttributesChanged[0].Notification.AddedAttributes.Find(LobbyMemberAttributeId2);
+			UTEST_NOT_NULL("Check LobbyMemberAttributesChanged is valid.", LobbyMemberAttribute2);
+			UTEST_EQUAL("Check LobbyMemberAttributesChanged is valid.", *LobbyMemberAttribute2, LobbyMemberAttribute2Value1);
+		}
+
+		// 4.4 setup - modify a lobby and lobby member attribute
 		{
-			FClientLobbyDataChanges LobbyChanges;
-			LobbyChanges.MutatedMembers.Add(User1, MakeShared<FClientLobbyMemberDataChanges>());
-			FApplyLobbyUpdateResult Result = ClientData.ApplyLobbyUpdateFromLocalChanges(MoveTemp(LobbyChanges), LobbyEvents);
-			UTEST_EQUAL("Check no local members left.", Result.LeavingLocalMembers.Num(), 0);
+			EventCapture.Empty();
+
+			FLobbyClientChanges LobbyChanges;
+			LobbyChanges.Attributes = { {{ LobbyAttributeId2, LobbyAttribute2Value2 }}, {} };
+			LobbyChanges.MemberAttributes = { {{ LobbyMemberAttributeId2, LobbyMemberAttribute2Value2 }}, {} };
+			TOnlineResult<FLobbyClientDataPrepareClientChanges> PrepareResult = ClientData.PrepareClientChanges({ User1, MoveTemp(LobbyChanges) });
+			UTEST_TRUE("Check PrepareClientChanges succeeded.", PrepareResult.IsOk());
+
+			FLobbyClientDataCommitClientChanges::Result CommitResult = ClientData.CommitClientChanges({ &LobbyEvents });
+			UTEST_EQUAL("Check no local members left.", CommitResult.LeavingLocalMembers.Num(), 0);
+		}
+
+		{
+			// 4.4.1. Local changes with changed lobby attribute triggers lobby attribute changed for changing attribute.
+			// 4.4.2. Local changes with changed lobby member attribute triggers lobby member attribute changed for changing attribute.
+			UTEST_EQUAL("Check events received.", EventCapture.GetTotalNotificationsReceived(), 2);
+			UTEST_EQUAL("Check LobbyAttributesChanged is valid.", EventCapture.LobbyAttributesChanged.Num(), 1);
+			UTEST_EQUAL("Check LobbyAttributesChanged is valid.", EventCapture.LobbyAttributesChanged[0].Notification.Lobby, ClientData.GetPublicDataPtr());
+			UTEST_EQUAL("Check LobbyAttributesChanged is valid.", EventCapture.LobbyAttributesChanged[0].Notification.AddedAttributes.Num(), 0);
+			UTEST_EQUAL("Check LobbyAttributesChanged is valid.", EventCapture.LobbyAttributesChanged[0].Notification.ChangedAttributes.Num(), 1);
+			UTEST_EQUAL("Check LobbyAttributesChanged is valid.", EventCapture.LobbyAttributesChanged[0].Notification.RemovedAttributes.Num(), 0);
+			const TPair<FSchemaVariant, FSchemaVariant>* LobbyAttribute2 = EventCapture.LobbyAttributesChanged[0].Notification.ChangedAttributes.Find(LobbyAttributeId2);
+			UTEST_NOT_NULL("Check LobbyAttributesChanged is valid.", LobbyAttribute2);
+			UTEST_EQUAL("Check LobbyAttributesChanged is valid.", *LobbyAttribute2, (TPair<FSchemaVariant, FSchemaVariant>(LobbyAttribute2Value1, LobbyAttribute2Value2)));
+			UTEST_EQUAL("Check LobbyMemberAttributesChanged is valid.", EventCapture.LobbyMemberAttributesChanged.Num(), 1);
+			UTEST_EQUAL("Check LobbyMemberAttributesChanged is valid.", EventCapture.LobbyMemberAttributesChanged[0].Notification.Lobby, ClientData.GetPublicDataPtr());
+			UTEST_EQUAL("Check LobbyMemberAttributesChanged is valid.", EventCapture.LobbyMemberAttributesChanged[0].Notification.Member->AccountId, User1);
+			UTEST_EQUAL("Check LobbyMemberAttributesChanged is valid.", EventCapture.LobbyMemberAttributesChanged[0].Notification.AddedAttributes.Num(), 0);
+			UTEST_EQUAL("Check LobbyMemberAttributesChanged is valid.", EventCapture.LobbyMemberAttributesChanged[0].Notification.ChangedAttributes.Num(), 1);
+			UTEST_EQUAL("Check LobbyMemberAttributesChanged is valid.", EventCapture.LobbyMemberAttributesChanged[0].Notification.RemovedAttributes.Num(), 0);
+			const TPair<FSchemaVariant, FSchemaVariant>* LobbyMemberAttribute2 = EventCapture.LobbyMemberAttributesChanged[0].Notification.ChangedAttributes.Find(LobbyMemberAttributeId2);
+			UTEST_NOT_NULL("Check LobbyMemberAttributesChanged is valid.", LobbyMemberAttribute2);
+			UTEST_EQUAL("Check LobbyMemberAttributesChanged is valid.", *LobbyMemberAttribute2, (TPair<FSchemaVariant, FSchemaVariant>(LobbyMemberAttribute2Value1, LobbyMemberAttribute2Value2)));
+		}
+
+		// 4.5 setup - remove a lobby and lobby member attribute
+		{
+			EventCapture.Empty();
+
+			FLobbyClientChanges LobbyChanges;
+			LobbyChanges.Attributes = { {}, { LobbyAttributeId2 } };
+			LobbyChanges.MemberAttributes = { {}, { LobbyMemberAttributeId2 } };
+			TOnlineResult<FLobbyClientDataPrepareClientChanges> PrepareResult = ClientData.PrepareClientChanges({ User1, MoveTemp(LobbyChanges) });
+			UTEST_TRUE("Check PrepareClientChanges succeeded.", PrepareResult.IsOk());
+
+			FLobbyClientDataCommitClientChanges::Result CommitResult = ClientData.CommitClientChanges({ &LobbyEvents });
+			UTEST_EQUAL("Check no local members left.", CommitResult.LeavingLocalMembers.Num(), 0);
 		}
 
-		// 4.2.1. Local members added to lobby through local change triggers lobby joined and lobby member joined. On member joined triggered for remote member.
-		UTEST_EQUAL("Check events received.", EventCapture.GetTotalNotificationsReceived(), 3);
-		UTEST_EQUAL("Check LobbyJoined valid.", EventCapture.LobbyJoined.Num(), 1);
-		UTEST_EQUAL("Check LobbyJoined valid.", EventCapture.LobbyJoined[0].GlobalIndex, 0);
-		UTEST_EQUAL("Check LobbyMemberJoined valid.", EventCapture.LobbyMemberJoined.Num(), 2);
-		FLobbyEventCapture::TNotificationInfo<FLobbyMemberJoined>* User1JoinNotification = EventCapture.LobbyMemberJoined[0].Notification.Member->AccountId == User1 ? &EventCapture.LobbyMemberJoined[0] : &EventCapture.LobbyMemberJoined[1];
-		FLobbyEventCapture::TNotificationInfo<FLobbyMemberJoined>* User2JoinNotification = EventCapture.LobbyMemberJoined[0].Notification.Member->AccountId == User2 ? &EventCapture.LobbyMemberJoined[0] : &EventCapture.LobbyMemberJoined[1];
-		UTEST_EQUAL("Check LobbyMemberJoined valid.", User1JoinNotification->Notification.Lobby, ClientData.GetPublicDataPtr());
-		UTEST_EQUAL("Check LobbyMemberJoined valid.", User1JoinNotification->Notification.Member->AccountId, User1);
-		UTEST_EQUAL("Check LobbyMemberJoined valid.", User1JoinNotification->Notification.Member->Attributes.Num(), 1);
-		UTEST_EQUAL("Check LobbyMemberJoined valid.", User2JoinNotification->Notification.Lobby, ClientData.GetPublicDataPtr());
-		UTEST_EQUAL("Check LobbyMemberJoined valid.", User2JoinNotification->Notification.Member->AccountId, User2);
+		{
+			// 4.5.1. Local changes with changed lobby attribute triggers lobby attribute changed for removing attribute.
+			// 4.5.2. Local changes with changed lobby member attribute triggers lobby member attribute changed for removing attribute.
+			UTEST_EQUAL("Check events received.", EventCapture.GetTotalNotificationsReceived(), 2);
+			UTEST_EQUAL("Check LobbyAttributesChanged is valid.", EventCapture.LobbyAttributesChanged.Num(), 1);
+			UTEST_EQUAL("Check LobbyAttributesChanged is valid.", EventCapture.LobbyAttributesChanged[0].Notification.Lobby, ClientData.GetPublicDataPtr());
+			UTEST_EQUAL("Check LobbyAttributesChanged is valid.", EventCapture.LobbyAttributesChanged[0].Notification.AddedAttributes.Num(), 0);
+			UTEST_EQUAL("Check LobbyAttributesChanged is valid.", EventCapture.LobbyAttributesChanged[0].Notification.ChangedAttributes.Num(), 0);
+			UTEST_EQUAL("Check LobbyAttributesChanged is valid.", EventCapture.LobbyAttributesChanged[0].Notification.RemovedAttributes.Num(), 1);
+			UTEST_TRUE("Check LobbyAttributesChanged is valid.", EventCapture.LobbyAttributesChanged[0].Notification.RemovedAttributes.Contains(LobbyAttributeId2));
+			UTEST_EQUAL("Check LobbyMemberAttributesChanged is valid.", EventCapture.LobbyMemberAttributesChanged.Num(), 1);
+			UTEST_EQUAL("Check LobbyMemberAttributesChanged is valid.", EventCapture.LobbyMemberAttributesChanged[0].Notification.Lobby, ClientData.GetPublicDataPtr());
+			UTEST_EQUAL("Check LobbyMemberAttributesChanged is valid.", EventCapture.LobbyMemberAttributesChanged[0].Notification.Member->AccountId, User1);
+			UTEST_EQUAL("Check LobbyMemberAttributesChanged is valid.", EventCapture.LobbyMemberAttributesChanged[0].Notification.AddedAttributes.Num(), 0);
+			UTEST_EQUAL("Check LobbyMemberAttributesChanged is valid.", EventCapture.LobbyMemberAttributesChanged[0].Notification.ChangedAttributes.Num(), 0);
+			UTEST_EQUAL("Check LobbyMemberAttributesChanged is valid.", EventCapture.LobbyMemberAttributesChanged[0].Notification.RemovedAttributes.Num(), 1);
+			UTEST_TRUE("Check LobbyMemberAttributesChanged is valid.", EventCapture.LobbyMemberAttributesChanged[0].Notification.RemovedAttributes.Contains(LobbyMemberAttributeId2));
+		}
 
-		// 4.3 setup
+		// 4.6 setup - assign leadership to another user.
 		{
 			EventCapture.Empty();
 
-			LobbyAttribute1Value.Set(TEXT("DoubleMutatedAttribute1"));
-			LobbyMemberAttribute1Value.Set(TEXT("DoubleMutatedMemberAttribute1"));
-			LobbySchemaName = TEXT("Waffles");
-
-			TSharedRef<FClientLobbyMemberDataChanges> User1Changes = MakeShared<FClientLobbyMemberDataChanges>();
-			User1Changes->MutatedAttributes.Add(LobbyMemberAttribute1Key, LobbyMemberAttribute1Value);
-
-			FClientLobbyDataChanges LobbyChanges;
-			LobbyChanges.MutatedAttributes.Add(LobbyAttribute1Key, LobbyAttribute1Value);
-			LobbyChanges.MutatedMembers.Add(User1, User1Changes);
-			LobbyChanges.LobbySchema = LobbySchemaName;
-			LobbyChanges.OwnerAccountId = User1;
-			FApplyLobbyUpdateResult Result = ClientData.ApplyLobbyUpdateFromLocalChanges(MoveTemp(LobbyChanges), LobbyEvents);
-			UTEST_EQUAL("Check no local members left.", Result.LeavingLocalMembers.Num(), 0);
-		}
-
-		// 4.3.1. Local changes with changed lobby attribute triggers lobby attribute changed.
-		// 4.3.2. Local changes with changed lobby member attribute triggers lobby member attribute changed.
-		// 4.3.3. Local changes with changed lobby schema triggers lobby schema changed.
-		// 4.3.4. Local changes with changed lobby leader triggers lobby leader changed.
-		UTEST_EQUAL("Check events received.", EventCapture.GetTotalNotificationsReceived(), 4);
-		UTEST_EQUAL("Check LobbyAttributesChanged is valid.", EventCapture.LobbyAttributesChanged.Num(), 1);
-		UTEST_EQUAL("Check LobbyAttributesChanged is valid.", EventCapture.LobbyAttributesChanged[0].Notification.Lobby, ClientData.GetPublicDataPtr());
-		UTEST_EQUAL("Check LobbyAttributesChanged is valid.", EventCapture.LobbyAttributesChanged[0].Notification.ChangedAttributes.Num(), 1);
-		UTEST_NOT_NULL("Check LobbyAttributesChanged is valid.", EventCapture.LobbyAttributesChanged[0].Notification.ChangedAttributes.Find(LobbyAttribute1Key));
-		UTEST_EQUAL("Check LobbyMemberAttributesChanged is valid.", EventCapture.LobbyMemberAttributesChanged.Num(), 1);
-		UTEST_EQUAL("Check LobbyMemberAttributesChanged is valid.", EventCapture.LobbyMemberAttributesChanged[0].Notification.Lobby, ClientData.GetPublicDataPtr());
-		UTEST_EQUAL("Check LobbyMemberAttributesChanged is valid.", EventCapture.LobbyMemberAttributesChanged[0].Notification.Member->AccountId, User1);
-		UTEST_EQUAL("Check LobbyMemberAttributesChanged is valid.", EventCapture.LobbyMemberAttributesChanged[0].Notification.ChangedAttributes.Num(), 1);
-		UTEST_NOT_NULL("Check LobbyMemberAttributesChanged is valid.", EventCapture.LobbyMemberAttributesChanged[0].Notification.ChangedAttributes.Find(LobbyMemberAttribute1Key));
-		UTEST_EQUAL("Check LobbySchemaChanged is valid.", EventCapture.LobbySchemaChanged.Num(), 1);
-		UTEST_EQUAL("Check LobbySchemaChanged is valid.", EventCapture.LobbySchemaChanged[0].Notification.Lobby, ClientData.GetPublicDataPtr());
-		UTEST_EQUAL("Check LobbyLeaderChanged is valid.", EventCapture.LobbyLeaderChanged.Num(), 1);
-		UTEST_EQUAL("Check LobbyLeaderChanged is valid.", EventCapture.LobbyLeaderChanged[0].Notification.Lobby, ClientData.GetPublicDataPtr());
-		UTEST_EQUAL("Check LobbyLeaderChanged is valid.", EventCapture.LobbyLeaderChanged[0].Notification.Leader->AccountId, User1);
-
-		// 4.4 setup
+			FLobbyClientChanges LobbyChanges;
+			LobbyChanges.OwnerAccountId = User2;
+			TOnlineResult<FLobbyClientDataPrepareClientChanges> PrepareResult = ClientData.PrepareClientChanges({ User1, MoveTemp(LobbyChanges) });
+			UTEST_TRUE("Check PrepareClientChanges succeeded.", PrepareResult.IsOk());
+
+			FLobbyClientDataCommitClientChanges::Result CommitResult = ClientData.CommitClientChanges({ &LobbyEvents });
+			UTEST_EQUAL("Check no local members left.", CommitResult.LeavingLocalMembers.Num(), 0);
+		}
+
+		{
+			// 4.6.1. Local changes with changed lobby leader triggers lobby leader changed.
+			UTEST_EQUAL("Check events received.", EventCapture.GetTotalNotificationsReceived(), 1);
+			UTEST_EQUAL("Check LobbyLeaderChanged is valid.", EventCapture.LobbyLeaderChanged.Num(), 1);
+			UTEST_EQUAL("Check LobbyLeaderChanged is valid.", EventCapture.LobbyLeaderChanged[0].Notification.Lobby, ClientData.GetPublicDataPtr());
+			UTEST_EQUAL("Check LobbyLeaderChanged is valid.", EventCapture.LobbyLeaderChanged[0].Notification.Leader->AccountId, User2);
+		}
+
+		// 4.7 Trying to both set member attributes and leave will fail.
+		{
+			FLobbyClientChanges LobbyChanges;
+			LobbyChanges.MemberAttributes = { {}, { LobbyMemberAttributeId1 } };
+			LobbyChanges.LocalUserLeaveReason = ELobbyMemberLeaveReason::Left;
+			TOnlineResult<FLobbyClientDataPrepareClientChanges> PrepareResult = ClientData.PrepareClientChanges({ User1, MoveTemp(LobbyChanges) });
+			UTEST_FALSE("Check PrepareClientChanges failed as expected.", PrepareResult.IsOk());
+		}
+
+		// 4.8 Trying to change lobby attributes while leaving will fail.
+		{
+			FLobbyClientChanges LobbyChanges;
+			LobbyChanges.Attributes = { {}, { LobbyAttributeId1 } };
+			LobbyChanges.LocalUserLeaveReason = ELobbyMemberLeaveReason::Left;
+			TOnlineResult<FLobbyClientDataPrepareClientChanges> PrepareResult = ClientData.PrepareClientChanges({ User1, MoveTemp(LobbyChanges) });
+			UTEST_FALSE("Check PrepareClientChanges failed as expected.", PrepareResult.IsOk());
+		}
+
+		// 4.9 Notifications are generated correctly for lobby cleanup when the last local member leaves.
 		{
 			EventCapture.Empty();
 
-			TSharedRef<FClientLobbyMemberDataChanges> User1Changes = MakeShared<FClientLobbyMemberDataChanges>();
-			User1Changes->ClearedAttributes.Add(LobbyMemberAttribute1Key);
-
-			FClientLobbyDataChanges LobbyChanges;
-			LobbyChanges.ClearedAttributes.Add(LobbyAttribute1Key);
-			LobbyChanges.MutatedMembers.Add(User1, User1Changes);
-			LobbyChanges.LeavingMembers.Add(User1, ELobbyMemberLeaveReason::Left);
-			FApplyLobbyUpdateResult Result = ClientData.ApplyLobbyUpdateFromLocalChanges(MoveTemp(LobbyChanges), LobbyEvents);
-			UTEST_EQUAL("Check local member left.", Result.LeavingLocalMembers.Num(), 1);
-			UTEST_EQUAL("Check local member left.", Result.LeavingLocalMembers[0], User1);
-		}
-
-		// 4.4.1. Local changes with cleared lobby attribute triggers lobby attribute changed.
-		// 4.4.2. Local changes with cleared member attribute triggers member attribute changed.
-		// 4.4.1. Local member removed through local change results in member leave for all members followed by lobby leave.
-		UTEST_EQUAL("Check events received.", EventCapture.GetTotalNotificationsReceived(), 5);
-		UTEST_EQUAL("Check LobbyAttributesChanged is valid.", EventCapture.LobbyAttributesChanged.Num(), 1);
-		UTEST_EQUAL("Check LobbyAttributesChanged is valid.", EventCapture.LobbyAttributesChanged[0].Notification.Lobby, ClientData.GetPublicDataPtr());
-		UTEST_EQUAL("Check LobbyAttributesChanged is valid.", EventCapture.LobbyAttributesChanged[0].Notification.ChangedAttributes.Num(), 1);
-		UTEST_NOT_NULL("Check LobbyAttributesChanged is valid.", EventCapture.LobbyAttributesChanged[0].Notification.ChangedAttributes.Find(LobbyAttribute1Key));
-		UTEST_EQUAL("Check LobbyMemberAttributesChanged is valid.", EventCapture.LobbyMemberAttributesChanged.Num(), 1);
-		UTEST_EQUAL("Check LobbyMemberAttributesChanged is valid.", EventCapture.LobbyMemberAttributesChanged[0].Notification.Lobby, ClientData.GetPublicDataPtr());
-		UTEST_EQUAL("Check LobbyMemberAttributesChanged is valid.", EventCapture.LobbyMemberAttributesChanged[0].Notification.Member->AccountId, User1);
-		UTEST_EQUAL("Check LobbyMemberAttributesChanged is valid.", EventCapture.LobbyMemberAttributesChanged[0].Notification.ChangedAttributes.Num(), 1);
-		UTEST_NOT_NULL("Check LobbyMemberAttributesChanged is valid.", EventCapture.LobbyMemberAttributesChanged[0].Notification.ChangedAttributes.Find(LobbyMemberAttribute1Key));
-		UTEST_EQUAL("Check LobbyMemberLeft valid.", EventCapture.LobbyMemberLeft.Num(), 2);
-		FLobbyEventCapture::TNotificationInfo<FLobbyMemberLeft>* User1LeaveNotification = EventCapture.LobbyMemberLeft[0].Notification.Member->AccountId == User1 ? &EventCapture.LobbyMemberLeft[0] : &EventCapture.LobbyMemberLeft[1];
-		FLobbyEventCapture::TNotificationInfo<FLobbyMemberLeft>* User2LeaveNotification = EventCapture.LobbyMemberLeft[0].Notification.Member->AccountId == User2 ? &EventCapture.LobbyMemberLeft[0] : &EventCapture.LobbyMemberLeft[1];
-		UTEST_EQUAL("Check LobbyMemberLeft valid.", User1LeaveNotification->Notification.Lobby, ClientData.GetPublicDataPtr());
-		UTEST_EQUAL("Check LobbyMemberLeft valid.", User1LeaveNotification->Notification.Member->AccountId, User1);
-		UTEST_EQUAL("Check LobbyMemberLeft valid.", User1LeaveNotification->Notification.Reason, ELobbyMemberLeaveReason::Left);
-		UTEST_EQUAL("Check LobbyMemberLeft valid.", User2LeaveNotification->Notification.Lobby, ClientData.GetPublicDataPtr());
-		UTEST_EQUAL("Check LobbyMemberLeft valid.", User2LeaveNotification->Notification.Member->AccountId, User2);
-		UTEST_EQUAL("Check LobbyMemberLeft valid.", User2LeaveNotification->Notification.Reason, ELobbyMemberLeaveReason::Left);
-		UTEST_EQUAL("Check LobbyLeft valid.", EventCapture.LobbyLeft.Num(), 1);
-		UTEST_EQUAL("Check LobbyLeft valid.", EventCapture.LobbyLeft[0].Notification.Lobby, ClientData.GetPublicDataPtr());
-		UTEST_EQUAL("Check LobbyLeft ordering.", EventCapture.LobbyLeft[0].GlobalIndex, EventCapture.GetTotalNotificationsReceived() - 1);
-		FLobbyEventCapture::TNotificationInfo<FLobbyMemberLeft>* LastMemberLeaveNotification = User1LeaveNotification->GlobalIndex > User2LeaveNotification->GlobalIndex ? User1LeaveNotification : User2LeaveNotification;
-		UTEST_EQUAL("Check LobbyMemberLeft ordering.", LastMemberLeaveNotification->GlobalIndex, EventCapture.LobbyLeft[0].GlobalIndex - 1);
+			FLobbyClientChanges LobbyChanges;
+			LobbyChanges.LocalUserLeaveReason = ELobbyMemberLeaveReason::Left;
+			TOnlineResult<FLobbyClientDataPrepareClientChanges> PrepareResult = ClientData.PrepareClientChanges({ User1, MoveTemp(LobbyChanges) });
+			UTEST_TRUE("Check PrepareClientChanges succeeded.", PrepareResult.IsOk());
+
+			FLobbyClientDataCommitClientChanges::Result CommitResult = ClientData.CommitClientChanges({ &LobbyEvents });
+			UTEST_EQUAL("Check local member left.", CommitResult.LeavingLocalMembers.Num(), 1);
+			UTEST_EQUAL("Check local member left.", CommitResult.LeavingLocalMembers[0], User1);
+		}
+
+		{
+			// 4.9.1. Local member removed through local change results in member leave for all members followed by lobby leave.
+			UTEST_EQUAL("Check events received.", EventCapture.GetTotalNotificationsReceived(), 3);
+			UTEST_EQUAL("Check LobbyMemberLeft valid.", EventCapture.LobbyMemberLeft.Num(), 2);
+			FLobbyEventCapture::TNotificationInfo<FLobbyMemberLeft>* User1LeaveNotification = EventCapture.LobbyMemberLeft[0].Notification.Member->AccountId == User1 ? &EventCapture.LobbyMemberLeft[0] : &EventCapture.LobbyMemberLeft[1];
+			FLobbyEventCapture::TNotificationInfo<FLobbyMemberLeft>* User2LeaveNotification = EventCapture.LobbyMemberLeft[0].Notification.Member->AccountId == User2 ? &EventCapture.LobbyMemberLeft[0] : &EventCapture.LobbyMemberLeft[1];
+			UTEST_EQUAL("Check LobbyMemberLeft valid.", User1LeaveNotification->Notification.Lobby, ClientData.GetPublicDataPtr());
+			UTEST_EQUAL("Check LobbyMemberLeft valid.", User1LeaveNotification->Notification.Member->AccountId, User1);
+			UTEST_EQUAL("Check LobbyMemberLeft valid.", User1LeaveNotification->Notification.Reason, ELobbyMemberLeaveReason::Left);
+			UTEST_EQUAL("Check LobbyMemberLeft valid.", User2LeaveNotification->Notification.Lobby, ClientData.GetPublicDataPtr());
+			UTEST_EQUAL("Check LobbyMemberLeft valid.", User2LeaveNotification->Notification.Member->AccountId, User2);
+			UTEST_EQUAL("Check LobbyMemberLeft valid.", User2LeaveNotification->Notification.Reason, ELobbyMemberLeaveReason::Left);
+			UTEST_EQUAL("Check LobbyLeft valid.", EventCapture.LobbyLeft.Num(), 1);
+			UTEST_EQUAL("Check LobbyLeft valid.", EventCapture.LobbyLeft[0].Notification.Lobby, ClientData.GetPublicDataPtr());
+			UTEST_EQUAL("Check LobbyLeft ordering.", EventCapture.LobbyLeft[0].GlobalIndex, EventCapture.GetTotalNotificationsReceived() - 1);
+			FLobbyEventCapture::TNotificationInfo<FLobbyMemberLeft>* LastMemberLeaveNotification = User1LeaveNotification->GlobalIndex > User2LeaveNotification->GlobalIndex ? User1LeaveNotification : User2LeaveNotification;
+			UTEST_EQUAL("Check LobbyMemberLeft ordering.", LastMemberLeaveNotification->GlobalIndex, EventCapture.LobbyLeft[0].GlobalIndex - 1);
+		}
 	}
 
 	// 5. Local changes supplemental
@@ -777,70 +1208,64 @@ bool FClientLobbyDataTest::RunTest(const FString& Parameters)
 		FLobbyEvents LobbyEvents;
 		FLobbyEventCapture EventCapture(LobbyEvents);
 		FLobbyId LobbyId(EOnlineServices::Null, 1);
-		FClientLobbyData ClientData(LobbyId);
+		FLobbyClientData ClientData(LobbyId, SchemaRegistry);
 
 		FAccountId User1(EOnlineServices::Null, 1);
 		FAccountId User2(EOnlineServices::Null, 2);
-		FName LobbySchemaName = TEXT("SchemaName");
 		int32 LobbyMaxMembers = 5;
 		ELobbyJoinPolicy LobbyJoinPolicy = ELobbyJoinPolicy::PublicAdvertised;
-		FLobbyAttributeId LobbyAttribute1Key = TEXT("LobbyAttribute1");
-		FLobbyVariant LobbyAttribute1Value;
-		LobbyAttribute1Value.Set(TEXT("Attribute1"));
-		FLobbyAttributeId LobbyAttribute2Key = TEXT("LobbyAttribute2");
-		FLobbyVariant LobbyAttribute2Value;
-		LobbyAttribute2Value.Set(TEXT("Attribute2"));
-		FLobbyAttributeId LobbyMemberAttribute1Key = TEXT("LobbyMemberAttribute1");
-		FLobbyVariant LobbyMemberAttribute1Value;
-		LobbyMemberAttribute1Value.Set(TEXT("MemberAttribute1"));
-		FLobbyAttributeId LobbyMemberAttribute2Key = TEXT("LobbyMemberAttribute2");
-		FLobbyVariant LobbyMemberAttribute2Value;
-		LobbyMemberAttribute2Value.Set(TEXT("MemberAttribute2"));
-
-		// 5.1 setup.
-		{
-			TMap<FAccountId, TSharedRef<FClientLobbyMemberSnapshot>> LobbyMemberSnapshots;
+
+		// 5.1 setup. Create a lobby from a snapshot containing two users.
+		{
+			TMap<FAccountId, FLobbyMemberServiceSnapshot> LobbyMemberSnapshots;
 			{
-				TSharedRef<FClientLobbyMemberSnapshot> MemberSnapshot = MakeShared<FClientLobbyMemberSnapshot>();
-				MemberSnapshot->AccountId = User1;
-				MemberSnapshot->Attributes.Add(LobbyMemberAttribute1Key, LobbyMemberAttribute1Value);
-				LobbyMemberSnapshots.Add(User1, MemberSnapshot);
+				FLobbyMemberServiceSnapshot& MemberSnapshot = LobbyMemberSnapshots.Add(User1);
+				MemberSnapshot.AccountId = User1;
+				MemberSnapshot.SchemaServiceSnapshot.Attributes.Add(LobbyMemberServiceAttributeId1, LobbyMemberAttribute1Value1);
 			}
 			{
-				TSharedRef<FClientLobbyMemberSnapshot> MemberSnapshot = MakeShared<FClientLobbyMemberSnapshot>();
-				MemberSnapshot->AccountId = User2;
-				LobbyMemberSnapshots.Add(User2, MemberSnapshot);
+				FLobbyMemberServiceSnapshot& MemberSnapshot = LobbyMemberSnapshots.Add(User2);
+				MemberSnapshot.AccountId = User2;
 			}
 
-			FClientLobbySnapshot LobbySnapshot;
+			FLobbyServiceSnapshot LobbySnapshot;
 			LobbySnapshot.OwnerAccountId = User1;
-			LobbySnapshot.SchemaName = LobbySchemaName;
 			LobbySnapshot.MaxMembers = LobbyMaxMembers;
 			LobbySnapshot.JoinPolicy = LobbyJoinPolicy;
-			LobbySnapshot.Attributes.Add(LobbyAttribute1Key, LobbyAttribute1Value);
+			LobbySnapshot.SchemaServiceSnapshot.Attributes.Add(LobbyServiceAttributeId1, SchemaCompatibilityAttributeValue);
+			LobbySnapshot.SchemaServiceSnapshot.Attributes.Add(LobbyServiceAttributeId2, LobbyAttribute1Value1);
 			LobbySnapshot.Members = {User1, User2};
 
 			TMap<FAccountId, ELobbyMemberLeaveReason> LeaveReasons;
-			FApplyLobbyUpdateResult Result = ClientData.ApplyLobbyUpdateFromServiceSnapshot(MoveTemp(LobbySnapshot), MoveTemp(LobbyMemberSnapshots), MoveTemp(LeaveReasons), &LobbyEvents);
-			UTEST_EQUAL("Check no local members left.", Result.LeavingLocalMembers.Num(), 0);
+			TOnlineResult<FLobbyClientDataPrepareServiceSnapshot> PrepareResult = ClientData.PrepareServiceSnapshot({ MoveTemp(LobbySnapshot), MoveTemp(LobbyMemberSnapshots), MoveTemp(LeaveReasons) });
+			UTEST_TRUE("Check PrepareServiceSnapshot succeeded.", PrepareResult.IsOk());
+
+			FLobbyClientDataCommitServiceSnapshot::Result CommitResult = ClientData.CommitServiceSnapshot({ &LobbyEvents });
+			UTEST_EQUAL("Check no local members left.", CommitResult.LeavingLocalMembers.Num(), 0);
 		}
 
-		// 5.1 setup continued.
+		// 5.1 setup continued. Join User 1 to the lobby assuming the service assigned role of owner.
 		{
-			FClientLobbyDataChanges LobbyChanges;
-			LobbyChanges.MutatedMembers.Add(User1, MakeShared<FClientLobbyMemberDataChanges>());
-			FApplyLobbyUpdateResult Result = ClientData.ApplyLobbyUpdateFromLocalChanges(MoveTemp(LobbyChanges), LobbyEvents);
-			UTEST_EQUAL("Check no local members left.", Result.LeavingLocalMembers.Num(), 0);
+			FLobbyClientChanges LobbyChanges;
+			LobbyChanges.MemberAttributes = { {}, {} };
+			TOnlineResult<FLobbyClientDataPrepareClientChanges> PrepareResult = ClientData.PrepareClientChanges({ User1, MoveTemp(LobbyChanges) });
+			UTEST_TRUE("Check PrepareClientChanges succeeded.", PrepareResult.IsOk());
+
+			FLobbyClientDataCommitClientChanges::Result CommitResult = ClientData.CommitClientChanges({ &LobbyEvents });
+			UTEST_EQUAL("Check no local members left.", CommitResult.LeavingLocalMembers.Num(), 0);
 		}
 
-		// 5.1 setup continued.
+		// 5.1 setup continued. User 1 kicks User 2.
 		{
 			EventCapture.Empty();
 
-			FClientLobbyDataChanges LobbyChanges;
-			LobbyChanges.LeavingMembers.Add(User2, ELobbyMemberLeaveReason::Kicked);
-			FApplyLobbyUpdateResult Result = ClientData.ApplyLobbyUpdateFromLocalChanges(MoveTemp(LobbyChanges), LobbyEvents);
-			UTEST_EQUAL("Check no local members left.", Result.LeavingLocalMembers.Num(), 0);
+			FLobbyClientChanges LobbyChanges;
+			LobbyChanges.KickedTargetMember = User2;
+			TOnlineResult<FLobbyClientDataPrepareClientChanges> PrepareResult = ClientData.PrepareClientChanges({ User1, MoveTemp(LobbyChanges) });
+			UTEST_TRUE("Check PrepareClientChanges succeeded.", PrepareResult.IsOk());
+
+			FLobbyClientDataCommitClientChanges::Result CommitResult = ClientData.CommitClientChanges({ &LobbyEvents });
+			UTEST_EQUAL("Check no local members left.", CommitResult.LeavingLocalMembers.Num(), 0);
 		}
 
 		// 5.1 Test local lobby owner kicking remote member.
@@ -1423,6 +1848,22 @@ END_ONLINE_STRUCT_META()
 
 /* Meta */ }
 
+/**
+ * Running the functional test below requires that the below config be added to DefaultEngine.ini.
+ * 
+ * [OnlineServices.Lobbies]
+ * +SchemaDescriptors=(Id="FunctionalTestLobbies", ParentId="LobbyBase")
+ * !SchemaCategoryAttributeDescriptors=ClearArray
+ * +SchemaCategoryAttributeDescriptors=(SchemaId="LobbyBase", CategoryId="Lobby", AttributeIds=("SchemaCompatibilityId", "LobbyCreateTime"))
+ * +SchemaCategoryAttributeDescriptors=(SchemaId="LobbyBase", CategoryId="LobbyMember")
+ * +SchemaCategoryAttributeDescriptors=(SchemaId="FunctionalTestLobbies", CategoryId="Lobby", AttributeIds=("LobbyAttributeId1"))
+ * +SchemaCategoryAttributeDescriptors=(SchemaId="FunctionalTestLobbies", CategoryId="LobbyMember", AttributeIds=("LobbyMemberAttributeId1"))
+ * +SchemaAttributeDescriptors=(Id="SchemaCompatibilityId", Type="Int64", Flags=("Public", "SchemaCompatibilityId"))
+ * +SchemaAttributeDescriptors=(Id="LobbyCreateTime", Type="Int64", Flags=("Public", "Searchable"))
+ * +SchemaAttributeDescriptors=(Id="LobbyAttributeId1", Type="String", Flags=("Public"), MaxSize=64)
+ * +SchemaAttributeDescriptors=(Id="LobbyMemberAttributeId1", Type="String", Flags=("Public"), MaxSize=64)
+ */
+
 TOnlineAsyncOpHandle<FFunctionalTestLobbies> RunLobbyFunctionalTest(IAuth& AuthInterface, FLobbiesCommon& LobbiesCommon, FLobbyEvents& LobbyEvents)
 {
 	static const FString ConfigName = TEXT("FunctionalTest");
@@ -1434,6 +1875,16 @@ TOnlineAsyncOpHandle<FFunctionalTestLobbies> RunLobbyFunctionalTest(IAuth& AuthI
 	static const FString ConfigNameKeyName = TEXT("Config");
 	static const FString SearchKeyName = TEXT("SearchKey");
 
+	static const FSchemaId LobbySchemaId = TEXT("FunctionalTestLobbies");
+	static const FSchemaAttributeId LobbyCreateTimeAttributeId = TEXT("LobbyCreateTime");
+	static const FSchemaAttributeId LobbyAttributeId1 = TEXT("LobbyAttributeId1");
+	static const FSchemaAttributeId LobbyMemberAttributeId1 = TEXT("LobbyMemberAttributeId1");
+	static const FSchemaVariant LobbyAttribute1Value1 = TEXT("LobbyAttribute1Value1");
+	static const FSchemaVariant LobbyAttribute1Value2 = TEXT("LobbyAttribute1Value2");
+	static const FSchemaVariant LobbyMemberAttribute1Value1 = TEXT("LobbyMemberAttribute1Value1");
+	static const FSchemaVariant LobbyMemberAttribute1Value2 = TEXT("LobbyMemberAttribute1Value2");
+	static const FName LocalLobbyName = TEXT("test");
+
 	struct FSearchParams
 	{
 		int64 LobbyCreateTime = 0;
@@ -1469,12 +1920,12 @@ TOnlineAsyncOpHandle<FFunctionalTestLobbies> RunLobbyFunctionalTest(IAuth& AuthI
 		// Create a lobby
 		FCreateLobby::Params Params;
 		Params.LocalAccountId = User1Info.AccountInfo->AccountId;
-		Params.LocalName = TEXT("test");
-		Params.SchemaName = TEXT("test");
+		Params.LocalName = LocalLobbyName;
+		Params.SchemaId = LobbySchemaId;
 		Params.MaxMembers = 2;
 		Params.JoinPolicy = ELobbyJoinPolicy::InvitationOnly;
-		Params.Attributes.Add(TEXT("test_attribute_key"), TEXT("test_attribute_value"));
-		Params.UserAttributes.Add(TEXT("test_attribute_key"), TEXT("test_attribute_value"));
+		Params.Attributes.Add(LobbyAttributeId1, LobbyAttribute1Value1);
+		Params.UserAttributes.Add(LobbyMemberAttributeId1, LobbyMemberAttribute1Value1);
 		return Params;
 	})
 	.Then(Private::CaptureOperationStepResult<FCreateLobby>(*Op, CreateLobbyKeyName, Private::FBindOperation<FCreateLobby>(&LobbiesCommon, &FLobbiesCommon::CreateLobby)))
@@ -1494,7 +1945,7 @@ TOnlineAsyncOpHandle<FFunctionalTestLobbies> RunLobbyFunctionalTest(IAuth& AuthI
 
 		// Check lobby attributes are set to the expected values.
 		if (!CreateResult.Lobby->Attributes.OrderIndependentCompareEqual(
-			TMap<FLobbyAttributeId, FLobbyVariant>{{TEXT("test_attribute_key"), TEXT("test_attribute_value")}}))
+			TMap<FSchemaAttributeId, FSchemaVariant>{{LobbyAttributeId1, LobbyAttribute1Value1}}))
 		{
 			InAsyncOp.SetError(Errors::Cancelled());
 		}
@@ -1503,7 +1954,7 @@ TOnlineAsyncOpHandle<FFunctionalTestLobbies> RunLobbyFunctionalTest(IAuth& AuthI
 		if (CreateResult.Lobby->Members.Num() == 1 && MemberData != nullptr)
 		{
 			if (!(**MemberData).Attributes.OrderIndependentCompareEqual(
-				TMap<FLobbyAttributeId, FLobbyVariant>{{TEXT("test_attribute_key"), TEXT("test_attribute_value")}}))
+				TMap<FSchemaAttributeId, FSchemaVariant>{{LobbyMemberAttributeId1, LobbyMemberAttribute1Value1}}))
 			{
 				InAsyncOp.SetError(Errors::Cancelled());
 			}
@@ -1524,7 +1975,7 @@ TOnlineAsyncOpHandle<FFunctionalTestLobbies> RunLobbyFunctionalTest(IAuth& AuthI
 		FModifyLobbyAttributes::Params Params;
 		Params.LocalAccountId = User1Info.AccountInfo->AccountId;
 		Params.LobbyId = CreateLobbyResult.Lobby->LobbyId;
-		Params.MutatedAttributes = {{TEXT("test_attribute_key"), TEXT("mutated_test_attribute_value")}};
+		Params.UpdatedAttributes = {{LobbyAttributeId1, LobbyAttribute1Value2}};
 		return Params;
 	})
 	.Then(Private::ConsumeOperationStepResult<FModifyLobbyAttributes>(*Op, Private::FBindOperation<FModifyLobbyAttributes>(&LobbiesCommon, &FLobbiesCommon::ModifyLobbyAttributes)))
@@ -1535,8 +1986,12 @@ TOnlineAsyncOpHandle<FFunctionalTestLobbies> RunLobbyFunctionalTest(IAuth& AuthI
 		if (EventCapture->LobbyAttributesChanged.Num() == 1)
 		{
 			const FLobbyAttributesChanged& Notification = EventCapture->LobbyAttributesChanged[0].Notification;
-			if (!Notification.ChangedAttributes.Difference(TSet<FLobbyAttributeId>{TEXT("test_attribute_key")}).IsEmpty() ||
-				!Notification.Lobby->Attributes.OrderIndependentCompareEqual(TMap<FLobbyAttributeId, FLobbyVariant>{{TEXT("test_attribute_key"), TEXT("mutated_test_attribute_value")}}))
+
+			TSet<FSchemaAttributeId> ChangedKeys;
+			Notification.ChangedAttributes.GetKeys(ChangedKeys);
+
+			if (!ChangedKeys.Difference(TSet<FSchemaAttributeId>{LobbyAttributeId1}).IsEmpty() ||
+				!Notification.Lobby->Attributes.OrderIndependentCompareEqual(TMap<FSchemaAttributeId, FSchemaVariant>{{LobbyAttributeId1, LobbyAttribute1Value2}}))
 			{
 				InAsyncOp.SetError(Errors::Cancelled());
 			}
@@ -1557,7 +2012,7 @@ TOnlineAsyncOpHandle<FFunctionalTestLobbies> RunLobbyFunctionalTest(IAuth& AuthI
 		FModifyLobbyMemberAttributes::Params Params;
 		Params.LocalAccountId = User1Info.AccountInfo->AccountId;
 		Params.LobbyId = CreateLobbyResult.Lobby->LobbyId;
-		Params.MutatedAttributes = {{TEXT("test_attribute_key"), TEXT("mutated_test_attribute_value")}};
+		Params.UpdatedAttributes = {{LobbyMemberAttributeId1, LobbyMemberAttribute1Value2}};
 		return Params;
 	})
 	.Then(Private::ConsumeOperationStepResult<FModifyLobbyMemberAttributes>(*Op, Private::FBindOperation<FModifyLobbyMemberAttributes>(&LobbiesCommon, &FLobbiesCommon::ModifyLobbyMemberAttributes)))
@@ -1572,8 +2027,11 @@ TOnlineAsyncOpHandle<FFunctionalTestLobbies> RunLobbyFunctionalTest(IAuth& AuthI
 			const FLobbyMemberAttributesChanged& Notification = EventCapture->LobbyMemberAttributesChanged[0].Notification;
 			if (const TSharedRef<const FLobbyMember>* MemberData = Notification.Lobby->Members.Find(User1Info.AccountInfo->AccountId))
 			{
-				if (!Notification.ChangedAttributes.Difference(TSet<FLobbyAttributeId>{TEXT("test_attribute_key")}).IsEmpty() ||
-					!(**MemberData).Attributes.OrderIndependentCompareEqual(TMap<FLobbyAttributeId, FLobbyVariant>{{TEXT("test_attribute_key"), TEXT("mutated_test_attribute_value")}}))
+				TSet<FSchemaAttributeId> ChangedKeys;
+				Notification.ChangedAttributes.GetKeys(ChangedKeys);
+
+				if (!ChangedKeys.Difference(TSet<FSchemaAttributeId>{LobbyMemberAttributeId1}).IsEmpty() ||
+					!(**MemberData).Attributes.OrderIndependentCompareEqual(TMap<FSchemaAttributeId, FSchemaVariant>{{LobbyMemberAttributeId1, LobbyMemberAttribute1Value2}}))
 				{
 					InAsyncOp.SetError(Errors::Cancelled());
 				}
@@ -1632,8 +2090,8 @@ TOnlineAsyncOpHandle<FFunctionalTestLobbies> RunLobbyFunctionalTest(IAuth& AuthI
 
 		FCreateLobby::Params Params;
 		Params.LocalAccountId = User1Info.AccountInfo->AccountId;
-		Params.LocalName = TEXT("test");
-		Params.SchemaName = TEXT("test");
+		Params.LocalName = LocalLobbyName;
+		Params.SchemaId = LobbySchemaId;
 		Params.MaxMembers = 2;
 		Params.JoinPolicy = ELobbyJoinPolicy::InvitationOnly;
 		//Params.Attributes;
@@ -1690,7 +2148,7 @@ TOnlineAsyncOpHandle<FFunctionalTestLobbies> RunLobbyFunctionalTest(IAuth& AuthI
 
 		FJoinLobby::Params Params;
 		Params.LocalAccountId = User2Info.AccountInfo->AccountId;
-		Params.LocalName = TEXT("test");
+		Params.LocalName = LocalLobbyName;
 		Params.LobbyId = CreateLobbyResult.Lobby->LobbyId;
 		return Params;
 	})
@@ -1779,11 +2237,11 @@ TOnlineAsyncOpHandle<FFunctionalTestLobbies> RunLobbyFunctionalTest(IAuth& AuthI
 
 		FCreateLobby::Params Params;
 		Params.LocalAccountId = User1Info.AccountInfo->AccountId;
-		Params.LocalName = TEXT("test");
-		Params.SchemaName = TEXT("test");
+		Params.LocalName = LocalLobbyName;
+		Params.SchemaId = LobbySchemaId;
 		Params.MaxMembers = 2;
 		Params.JoinPolicy = ELobbyJoinPolicy::PublicAdvertised;
-		Params.Attributes = {{ TEXT("LobbyCreateTime"), SearchParams->LobbyCreateTime }};
+		Params.Attributes = {{ LobbyCreateTimeAttributeId, SearchParams->LobbyCreateTime }};
 		return Params;
 	})
 	.Then(Private::CaptureOperationStepResult<FCreateLobby>(*Op, CreateLobbyKeyName, Private::FBindOperation<FCreateLobby>(&LobbiesCommon, &FLobbiesCommon::CreateLobby)))
@@ -1813,7 +2271,7 @@ TOnlineAsyncOpHandle<FFunctionalTestLobbies> RunLobbyFunctionalTest(IAuth& AuthI
 		// Search for lobby from create. Searching by attribute will also limit results by bucket id.
 		FFindLobbies::Params Params;
 		Params.LocalAccountId = User2Info.AccountInfo->AccountId;
-		Params.Filters = {{TEXT("LobbyCreateTime"), ELobbyComparisonOp::Equals, SearchParams->LobbyCreateTime}};
+		Params.Filters = {{LobbyCreateTimeAttributeId, ESchemaAttributeComparisonOp::Equals, SearchParams->LobbyCreateTime}};
 		return Params;
 	})
 	.Then(Private::CaptureOperationStepResult<FFindLobbies>(*Op, FindLobbyKeyName, Private::FBindOperation<FFindLobbies>(&LobbiesCommon, &FLobbiesCommon::FindLobbies)))
@@ -1834,7 +2292,7 @@ TOnlineAsyncOpHandle<FFunctionalTestLobbies> RunLobbyFunctionalTest(IAuth& AuthI
 
 		FJoinLobby::Params Params;
 		Params.LocalAccountId = User2Info.AccountInfo->AccountId;
-		Params.LocalName = TEXT("test");
+		Params.LocalName = LocalLobbyName;
 		Params.LobbyId = FindResults.Lobbies[0]->LobbyId;
 		return Params;
 	})
