Integrate schema into lobbies.

diff --git a/Engine/Plugins/Online/Experimental/OnlineServices/Source/OnlineServicesCommon/Private/Online/LobbiesCommonTypes.cpp b/Engine/Plugins/Online/Experimental/OnlineServices/Source/OnlineServicesCommon/Private/Online/LobbiesCommonTypes.cpp
--- a/Engine/Plugins/Online/Experimental/OnlineServices/Source/OnlineServicesCommon/Private/Online/LobbiesCommonTypes.cpp
+++ b/Engine/Plugins/Online/Experimental/OnlineServices/Source/OnlineServicesCommon/Private/Online/LobbiesCommonTypes.cpp
@@ -4,517 +4,704 @@
 
 namespace UE::Online {
 
-TSharedPtr<FLobbySchema> FLobbySchema::Create(FLobbySchemaConfig LobbySchemaConfig)
-{
-	TSharedRef<FLobbySchema> Schema = MakeShared<FLobbySchema>();
-	Schema->SchemaId = LobbySchemaConfig.SchemaName;
-
-	// Todo: for real validation and init
-	if (Schema->SchemaId == FLobbySchemaId())
-	{
-		return TSharedPtr<FLobbySchema>();
-	}
-	else
-	{
-		return Schema;
-	}
-}
+FSchemaId LobbyBaseSchemaId = TEXT("LobbyBase");
+FSchemaCategoryId LobbySchemaCategoryId = TEXT("Lobby");
+FSchemaCategoryId LobbyMemberSchemaCategoryId = TEXT("LobbyMember");
 
-TSharedRef<FLobbyServiceAttributeChanges> FLobbySchema::TranslateLobbyAttributes(const TMap<FLobbyAttributeId, FLobbyVariant>& LobbyAttributes, const FLobbyClientAttributeChanges& ClientAttributeChanges) const
+FLobbyClientData::FLobbyClientData(FLobbyId LobbyId, const TSharedRef<const FSchemaRegistry>& InSchemaRegistry)
+	: SchemaRegistry(InSchemaRegistry)
+	, InternalPublicData(MakeShared<FLobbyInternal>(InSchemaRegistry))
 {
-	// Todo: implementation.
-
-	TSharedRef<FLobbyServiceAttributeChanges> Changes = MakeShared<FLobbyServiceAttributeChanges>();
-	Changes->MutatedAttributes = ClientAttributeChanges.MutatedAttributes;
-	Changes->ClearedAttributes = ClientAttributeChanges.ClearedAttributes;
-	return Changes;
+	InternalPublicData->LobbyId = LobbyId;
 }
 
-TSharedRef<FLobbyClientAttributeChanges> FLobbySchema::TranslateLobbyAttributes(const FLobbyServiceAttributeChanges& ServiceAttributeChanges) const
+TSharedPtr<const FLobbyMemberInternal> FLobbyClientData::GetMemberData(FAccountId MemberAccountId) const
 {
-	// Todo: implementation.
-
-	TSharedRef<FLobbyClientAttributeChanges> Changes = MakeShared<FLobbyClientAttributeChanges>();
-	Changes->MutatedAttributes = ServiceAttributeChanges.MutatedAttributes;
-	Changes->ClearedAttributes = ServiceAttributeChanges.ClearedAttributes;
-	return Changes;
+	const TSharedRef<FLobbyMemberInternal>* FoundMemberData = MemberDataStorage.Find(MemberAccountId);
+	return FoundMemberData ? *FoundMemberData : TSharedPtr<const FLobbyMemberInternal>();
 }
 
-TSharedRef<FLobbyServiceAttributeChanges> FLobbySchema::TranslateLobbyMemberAttributes(const TMap<FLobbyAttributeId, FLobbyVariant>& LobbyAttributes, const FLobbyClientAttributeChanges& ClientAttributeChanges) const
+TOnlineResult<FLobbyClientDataPrepareServiceSnapshot> FLobbyClientData::PrepareServiceSnapshot(FLobbyClientDataPrepareServiceSnapshot::Params&& Params) const
 {
-	// Todo: implementation.
-
-	TSharedRef<FLobbyServiceAttributeChanges> Changes = MakeShared<FLobbyServiceAttributeChanges>();
-	Changes->MutatedAttributes = ClientAttributeChanges.MutatedAttributes;
-	Changes->ClearedAttributes = ClientAttributeChanges.ClearedAttributes;
-	return Changes;
-}
+	FPreparedServiceChanges NewPreparedServiceChanges;
 
-TSharedRef<FLobbyClientAttributeChanges> FLobbySchema::TranslateLobbyMemberAttributes(const FLobbyServiceAttributeChanges& ServiceAttributeChanges) const
-{
-	// Todo: implementation.
+	// Reset any pending changes. No other changes can occur while applying a service snapshot.
+	ResetPreparedChanges();
 
-	TSharedRef<FLobbyClientAttributeChanges> Changes = MakeShared<FLobbyClientAttributeChanges>();
-	Changes->MutatedAttributes = ServiceAttributeChanges.MutatedAttributes;
-	Changes->ClearedAttributes = ServiceAttributeChanges.ClearedAttributes;
-	return Changes;
-}
+	// Prepare lobby attribute changes.
+	TOnlineResult<FSchemaCategoryInstancePrepareServiceSnapshot> PrepareLobbyServiceSnapshotResult =
+		InternalPublicData->SchemaData.PrepareServiceSnapshot({ TOptional<FSchemaId>(), MoveTemp(Params.LobbySnapshot.SchemaServiceSnapshot) });
+	if (PrepareLobbyServiceSnapshotResult.IsError())
+	{
+		UE_LOG(LogTemp, Warning, TEXT("[FLobbyClientData::PrepareServiceSnapshot] Failed to prepare lobby service snapshot: Lobby[%s], Result[%s]"),
+			*ToLogString(InternalPublicData->LobbyId), *PrepareLobbyServiceSnapshotResult.GetErrorValue().GetLogString());
+		return TOnlineResult<FLobbyClientDataPrepareServiceSnapshot>(MoveTemp(PrepareLobbyServiceSnapshotResult.GetErrorValue()));
+	}
 
-bool FLobbySchemaRegistry::Initialize(TArray<FLobbySchemaConfig> LobbySchemaConfigs)
-{
-	bool bSuccess = true;
+	// Get schema id to use when applying member changes.
+	// Preparing the lobby snapshot may have resulted in a schema change.
+	TOptional<FSchemaId> NewDerivedSchemaId = PrepareLobbyServiceSnapshotResult.GetOkValue().DerivedSchemaId;
 
-	for (FLobbySchemaConfig& LobbySchemaConfig : LobbySchemaConfigs)
+	// Process members who left.
+	// Build list of leaving members by iterating over the known members.
+	for (const TPair<FAccountId, TSharedRef<FLobbyMemberInternal>>& MemberData : MemberDataStorage)
 	{
-		bSuccess &= RegisterSchema(MoveTemp(LobbySchemaConfig));
+		// Check if member exists in new snapshot.
+		if (!Params.LobbySnapshot.Members.Contains(MemberData.Key))
+		{
+			// Member has left - check if they gave an explicit reason.
+			if (const ELobbyMemberLeaveReason* MemberLeaveReason = Params.LeaveReasons.Find(MemberData.Key))
+			{
+				NewPreparedServiceChanges.LeavingMembers.Add({MemberData.Value, *MemberLeaveReason});
+			}
+			else
+			{
+				const ELobbyMemberLeaveReason DefaultLeaveReason = ELobbyMemberLeaveReason::Disconnected;
+				NewPreparedServiceChanges.LeavingMembers.Add({ MemberData.Value, DefaultLeaveReason });
+				UE_LOG(LogTemp, Warning, TEXT("[FLobbyClientData::PrepareServiceSnapshot] Member left lobby without giving reason, defaulting to %s: Lobby[%s], Member[%s]"),
+					*ToLogString(InternalPublicData->LobbyId), LexToString(DefaultLeaveReason), *ToLogString(MemberData.Key));
+			}
+		}
 	}
 
-	return bSuccess;
-}
-
-TSharedPtr<FLobbySchema> FLobbySchemaRegistry::FindSchema(FLobbySchemaId SchemaId)
-{
-	TSharedRef<FLobbySchema>* Schema = RegisteredSchemas.Find(SchemaId);
-	return Schema ? *Schema : TSharedPtr<FLobbySchema>();
-}
+	// Process member joins and attribute changes.
+	// Joining members are expected to have a member data snapshot.
+	for (TPair<FAccountId, FLobbyMemberServiceSnapshot>& MemberSnapshotPair : Params.LobbyMemberSnapshots)
+	{
+		if (!Params.LobbySnapshot.Members.Contains(MemberSnapshotPair.Key))
+		{
+			UE_LOG(LogTemp, Warning, TEXT("[FLobbyClientData::PrepareServiceSnapshot] Member update ignored for unknown member: Lobby[%s], Member[%s]"),
+				*ToLogString(InternalPublicData->LobbyId), *ToLogString(MemberSnapshotPair.Key));
+		}
+		else
+		{
+			FLobbyMemberServiceSnapshot& MemberSnapshot = MemberSnapshotPair.Value;
+			FLobbyMemberInternal* PreparingMemberData = nullptr;
 
-bool FLobbySchemaRegistry::RegisterSchema(FLobbySchemaConfig LobbySchemaConfig)
-{
-	// Todo: handle schema hierarchy.
-	// Todo: handle schema parsing.
+			// Check whether this member is already in the lobby.
+			if (const TSharedRef<FLobbyMemberInternal>* MemberData = MemberDataStorage.Find(MemberSnapshot.AccountId))
+			{
+				NewPreparedServiceChanges.UpdatedMemberDataStorage.Add(*MemberData);
+				PreparingMemberData = &MemberData->Get();
+			}
+			else
+			{
+				// New member, create an entry for them.
+				TSharedRef<FLobbyMemberInternal> NewMemberData = MakeShared<FLobbyMemberInternal>(SchemaRegistry);
+				NewMemberData->AccountId = MemberSnapshot.AccountId;
+				NewMemberData->PlatformAccountId = MemberSnapshot.PlatformAccountId;
+				NewMemberData->PlatformDisplayName = MemberSnapshot.PlatformDisplayName;
+				NewPreparedServiceChanges.NewMemberDataStorage.Add(NewMemberData);
+				PreparingMemberData = &NewMemberData.Get();
+			}
 
-	FLobbySchemaId SchemaName = LobbySchemaConfig.SchemaName;
+			// Try to prepare the lobby member snapshot.
+			TOnlineResult<FSchemaCategoryInstancePrepareServiceSnapshot> PrepareLobbyMemberServiceSnapshotResult =
+				PreparingMemberData->SchemaData.PrepareServiceSnapshot({ NewDerivedSchemaId, MoveTemp(MemberSnapshot.SchemaServiceSnapshot) });
+			if (PrepareLobbyMemberServiceSnapshotResult.IsError())
+			{
+				UE_LOG(LogTemp, Warning, TEXT("[FLobbyClientData::PrepareServiceSnapshot] Failed to prepare lobby member service snapshot: Lobby[%s], Member[%s], Result[%s]"),
+					*ToLogString(InternalPublicData->LobbyId), *ToLogString(MemberSnapshot.AccountId), *PrepareLobbyMemberServiceSnapshotResult.GetErrorValue().GetLogString());
+				return TOnlineResult<FLobbyClientDataPrepareServiceSnapshot>(MoveTemp(PrepareLobbyMemberServiceSnapshotResult.GetErrorValue()));
+			}
+		}
+	}
 
-	if (TSharedPtr<FLobbySchema> Schema = FLobbySchema::Create(MoveTemp(LobbySchemaConfig)))
+	// Check whether lobby ownership changed and is valid.
+	if (InternalPublicData->OwnerAccountId != Params.LobbySnapshot.OwnerAccountId)
 	{
-		RegisteredSchemas.Emplace(SchemaName, Schema.ToSharedRef());
-		return true;
+		if (!Params.LobbySnapshot.Members.Contains(Params.LobbySnapshot.OwnerAccountId))
+		{
+			UE_LOG(LogTemp, Warning, TEXT("[FLobbyClientData::PrepareServiceSnapshot] Lobby owner chage failed - lobby member data not found: Lobby[%s], Member[%s]"),
+				*ToLogString(InternalPublicData->LobbyId), *ToLogString(Params.LobbySnapshot.OwnerAccountId));
+			return TOnlineResult<FLobbyClientDataPrepareServiceSnapshot>(Errors::InvalidState());
+		}
+
+		NewPreparedServiceChanges.OwnerAccountId = Params.LobbySnapshot.OwnerAccountId;
 	}
-	else
+
+	// Check for max members change.
+	if (InternalPublicData->MaxMembers != Params.LobbySnapshot.MaxMembers)
 	{
-		return false;
+		NewPreparedServiceChanges.MaxMembers = Params.LobbySnapshot.MaxMembers;
 	}
-}
 
+	// Check for join policy change.
+	if (InternalPublicData->JoinPolicy != Params.LobbySnapshot.JoinPolicy)
+	{
+		NewPreparedServiceChanges.JoinPolicy = Params.LobbySnapshot.JoinPolicy;
+	}
+	
+	// Store prepared changes.
+	PreparedServiceChanges = MoveTemp(NewPreparedServiceChanges);
 
-FClientLobbyData::FClientLobbyData(FLobbyId LobbyId)
-	: PublicData(MakeShared<FLobby>())
-{
-	PublicData->LobbyId = LobbyId;
+	return TOnlineResult<FLobbyClientDataPrepareServiceSnapshot>(FLobbyClientDataPrepareServiceSnapshot::Result{});
 }
 
-TSharedPtr<const FClientLobbyMemberSnapshot> FClientLobbyData::GetMemberData(FAccountId MemberAccountId) const
+FLobbyClientDataCommitServiceSnapshot::Result FLobbyClientData::CommitServiceSnapshot(FLobbyClientDataCommitServiceSnapshot::Params&& Params)
 {
-	const TSharedRef<FClientLobbyMemberSnapshot>* FoundMemberData = MemberDataStorage.Find(MemberAccountId);
-	return FoundMemberData ? *FoundMemberData : TSharedPtr<const FClientLobbyMemberSnapshot>();
-}
+	FLobbyClientDataCommitServiceSnapshot::Result CommitResult;
+	const bool bDispatchNotifications = Params.LobbyEvents != nullptr && !LocalMembers.IsEmpty();
 
-FApplyLobbyUpdateResult FClientLobbyData::ApplyLobbyUpdateFromServiceSnapshot(
-	FClientLobbySnapshot&& LobbySnapshot,
-	TMap<FAccountId, TSharedRef<FClientLobbyMemberSnapshot>>&& LobbyMemberSnapshots,
-	TMap<FAccountId, ELobbyMemberLeaveReason>&& LeaveReasons,
-	FLobbyEvents* LobbyEvents)
-{
-	FApplyLobbyUpdateResult Result;
-	const bool bDispatchNotifications = LobbyEvents != nullptr && !LocalMembers.IsEmpty();
+	if (!PreparedServiceChanges)
+	{
+		UE_LOG(LogTemp, Warning, TEXT("[FLobbyClientData::CommitServiceSnapshot] Commit failed, no pending changes found: Lobby[%s]"),
+			*ToLogString(InternalPublicData->LobbyId));
+		return CommitResult;
+	}
+
+	// Commit schema changes.
+	FSchemaCategoryInstanceCommitServiceSnapshot::Result CommitLobbyServiceSnapshotResult = InternalPublicData->SchemaData.CommitServiceSnapshot();
+	FSchemaServiceClientChanges& LobbyServiceClientChanges = CommitLobbyServiceSnapshotResult.ClientChanges;
 
 	// Notify if schema changed.
 	// Schema change notification must be processed first as the schema affects how the client processes attributes.
-	if (LobbySnapshot.SchemaName != PublicData->SchemaName)
+	if (LobbyServiceClientChanges.SchemaId)
 	{
-		PublicData->SchemaName = LobbySnapshot.SchemaName;
+		FLobbySchemaChanged LobbySchemaChanged = { InternalPublicData, InternalPublicData->SchemaId, *LobbyServiceClientChanges.SchemaId };
+		InternalPublicData->SchemaId = *LobbyServiceClientChanges.SchemaId;
 		if (bDispatchNotifications)
 		{
-			LobbyEvents->OnLobbySchemaChanged.Broadcast(FLobbySchemaChanged{PublicData});
+			Params.LobbyEvents->OnLobbySchemaChanged.Broadcast(LobbySchemaChanged);
 		}
 	}
 
-	// Handle lobby attribute changes.
-	TSet<FLobbyAttributeId> LobbyAttributeChanges = ApplyAttributeUpdateFromSnapshot(MoveTemp(LobbySnapshot.Attributes), PublicData->Attributes);
-	if (bDispatchNotifications && !LobbyAttributeChanges.IsEmpty())
+	const bool bAnyLobbyAttributesChanged =
+		!LobbyServiceClientChanges.AddedAttributes.IsEmpty() ||
+		!LobbyServiceClientChanges.ChangedAttributes.IsEmpty() ||
+		!LobbyServiceClientChanges.RemovedAttributes.IsEmpty();
+
+	// Handle lobby attribute change notifications.
+	if (bDispatchNotifications && bAnyLobbyAttributesChanged)
 	{
-		LobbyEvents->OnLobbyAttributesChanged.Broadcast(FLobbyAttributesChanged{PublicData, MoveTemp(LobbyAttributeChanges)});
+		Params.LobbyEvents->OnLobbyAttributesChanged.Broadcast({
+			InternalPublicData,
+			MoveTemp(LobbyServiceClientChanges.AddedAttributes),
+			MoveTemp(LobbyServiceClientChanges.ChangedAttributes),
+			MoveTemp(LobbyServiceClientChanges.RemovedAttributes) });
 	}
 
-	// Process members who left.
+	// Handle joining members.
+	for (TSharedRef<FLobbyMemberInternal> NewMember : PreparedServiceChanges->NewMemberDataStorage)
 	{
-		// Build list of leaving members.
-		TArray<TPair<TSharedRef<FClientLobbyMemberSnapshot>, ELobbyMemberLeaveReason>> LeavingMembers;
-		for (TPair<FAccountId, TSharedRef<FClientLobbyMemberSnapshot>>& MemberData : MemberDataStorage)
-		{
-			// Check if member exists in new snapshot.
-			if (LobbySnapshot.Members.Find(MemberData.Key) == nullptr)
-			{
-				// Member has left - check if they gave an explicit reason.
-				ELobbyMemberLeaveReason ResovedReason = ELobbyMemberLeaveReason::Disconnected;
-				if (const ELobbyMemberLeaveReason* MemberLeaveReason = LeaveReasons.Find(MemberData.Key))
-				{
-					ResovedReason = *MemberLeaveReason;
-				}
-				else
-				{
-					UE_LOG(LogTemp, Warning, TEXT("[FClientLobbyData::ApplyLobbyUpdateFromSnapshot] Member left lobby without giving reason: Lobby: %s, Member: %s"),
-						*ToLogString(PublicData->LobbyId), *ToLogString(MemberData.Key));
-				}
+		// Commit attribute changes.
+		NewMember->SchemaData.CommitServiceSnapshot();
+		
+		// todo.
+		//NewMember->PlatformAccountId = ;
+		//NewMember->PlatformDisplayName = ;
 
-				LeavingMembers.Emplace(MemberData.Value, ResovedReason);
-			}
-		}
+		// Add member to lobby structures.
+		AddMember(NewMember);
 
-		// Remove member data and notify.
-		for (const TPair<TSharedRef<FClientLobbyMemberSnapshot>, ELobbyMemberLeaveReason>& MemberData : LeavingMembers)
+		// Notify.
+		if (bDispatchNotifications)
 		{
-			MemberDataStorage.Remove(MemberData.Key->AccountId);
-			PublicData->Members.Remove(MemberData.Key->AccountId);
+			Params.LobbyEvents->OnLobbyMemberJoined.Broadcast({ InternalPublicData, NewMember });
+		}
+	}
 
-			if (MemberData.Key->bIsLocalMember)
-			{
-				LocalMembers.Remove(MemberData.Key->AccountId);
-				Result.LeavingLocalMembers.Add(MemberData.Key->AccountId);
-			}
+	// Handle updated members.
+	for (TSharedRef<FLobbyMemberInternal> UpdatedMember : PreparedServiceChanges->UpdatedMemberDataStorage)
+	{
+		// Commit attribute changes.
+		FSchemaCategoryInstanceCommitServiceSnapshot::Result CommitLobbyMemberServiceSnapshotResult = UpdatedMember->SchemaData.CommitServiceSnapshot();
+		FSchemaServiceClientChanges& LobbyMemberServiceClientChanges = CommitLobbyMemberServiceSnapshotResult.ClientChanges;
 
-			if (bDispatchNotifications)
-			{
-				LobbyEvents->OnLobbyMemberLeft.Broadcast(FLobbyMemberLeft{PublicData, MemberData.Key, MemberData.Value});
-			}
+		const bool bAnyLobbyMemberAttributesChanged =
+			!LobbyMemberServiceClientChanges.AddedAttributes.IsEmpty() ||
+			!LobbyMemberServiceClientChanges.ChangedAttributes.IsEmpty() ||
+			!LobbyMemberServiceClientChanges.RemovedAttributes.IsEmpty();
+
+		// Handle lobby attribute changes.
+		if (bDispatchNotifications && bAnyLobbyMemberAttributesChanged)
+		{
+			Params.LobbyEvents->OnLobbyMemberAttributesChanged.Broadcast({
+				InternalPublicData,
+				UpdatedMember,
+				MoveTemp(LobbyMemberServiceClientChanges.AddedAttributes),
+				MoveTemp(LobbyMemberServiceClientChanges.ChangedAttributes),
+				MoveTemp(LobbyMemberServiceClientChanges.RemovedAttributes) });
 		}
 	}
 
-	// Process member joins and attribute changes.
-	// Joining members are expected to have a member data snapshot.
-	for (TPair<FAccountId, TSharedRef<FClientLobbyMemberSnapshot>>& MemberSnapshot : LobbyMemberSnapshots)
+	// Handle leaving members.
+	for (TPair<TSharedRef<FLobbyMemberInternal>, ELobbyMemberLeaveReason> LeavingMemberPair : PreparedServiceChanges->LeavingMembers)
 	{
-		if (LobbySnapshot.Members.Find(MemberSnapshot.Key) == nullptr)
+		RemoveMember(LeavingMemberPair.Key);
+
+		if (LeavingMemberPair.Key->bIsLocalMember)
 		{
-			UE_LOG(LogTemp, Warning, TEXT("[FClientLobbyData::ApplyLobbyUpdateFromSnapshot] Member update ignored for unknown member: Lobby: %s, Member: %s"),
-				*ToLogString(PublicData->LobbyId), *ToLogString(MemberSnapshot.Key));
+			CommitResult.LeavingLocalMembers.Add(LeavingMemberPair.Key->AccountId);
 		}
-		else
+
+		if (bDispatchNotifications)
 		{
-			// Check whether this member is already in the lobby.
-			if (TSharedRef<FClientLobbyMemberSnapshot>* MemberData = MemberDataStorage.Find(MemberSnapshot.Key))
-			{
-				// Member is already in the lobby. Check if their attributes changed.
-				TSet<FLobbyAttributeId> MemberAttributeChanges = ApplyAttributeUpdateFromSnapshot(MoveTemp(MemberSnapshot.Value->Attributes), (*MemberData)->Attributes);
-				if (bDispatchNotifications && !MemberAttributeChanges.IsEmpty())
-				{
-					LobbyEvents->OnLobbyMemberAttributesChanged.Broadcast(FLobbyMemberAttributesChanged{PublicData, (*MemberData), MemberAttributeChanges});
-				}
-			}
-			else
-			{
-				// Member is not in the lobby, add them.
-				MemberDataStorage.Add(MemberSnapshot.Key, MemberSnapshot.Value);
-				PublicData->Members.Add(MemberSnapshot.Key, MoveTemp(MemberSnapshot.Value));
-				if (bDispatchNotifications)
-				{
-					LobbyEvents->OnLobbyMemberJoined.Broadcast(FLobbyMemberJoined{PublicData, MemberSnapshot.Value});
-				}
-			}
+			Params.LobbyEvents->OnLobbyMemberLeft.Broadcast({ InternalPublicData, LeavingMemberPair.Key, LeavingMemberPair.Value });
 		}
 	}
 
-	// Process lobby state.
-	// The lobby owner notification must be processed after member updates as the new owner may not
-	// have existed in the previous snapshot.
-	if (PublicData->OwnerAccountId != LobbySnapshot.OwnerAccountId)
+	// Check for ownership change.
+	if (PreparedServiceChanges->OwnerAccountId)
 	{
-		if (TSharedRef<FClientLobbyMemberSnapshot>* MemberData = MemberDataStorage.Find(LobbySnapshot.OwnerAccountId))
-		{
-			PublicData->OwnerAccountId = LobbySnapshot.OwnerAccountId;
-			if (bDispatchNotifications)
-			{
-				LobbyEvents->OnLobbyLeaderChanged.Broadcast(FLobbyLeaderChanged{PublicData, (*MemberData)});
-			}
-		}
-		else
+		InternalPublicData->OwnerAccountId = *PreparedServiceChanges->OwnerAccountId;
+
+		if (bDispatchNotifications)
 		{
-			UE_LOG(LogTemp, Warning, TEXT("[FClientLobbyData::ApplyLobbyUpdateFromSnapshot] Lobby owner chage failed - lobby member data not found: Lobby: %s, Member: %s"),
-				*ToLogString(PublicData->LobbyId), *ToLogString(LobbySnapshot.OwnerAccountId));
+			Params.LobbyEvents->OnLobbyLeaderChanged.Broadcast(FLobbyLeaderChanged{ InternalPublicData, MemberDataStorage.FindChecked(InternalPublicData->OwnerAccountId) });
 		}
 	}
 
-	// Set remaining lobby data.
-	PublicData->MaxMembers = LobbySnapshot.MaxMembers;
-	PublicData->JoinPolicy = LobbySnapshot.JoinPolicy;
+	// Check for max member change.
+	if (PreparedServiceChanges->MaxMembers)
+	{
+		InternalPublicData->MaxMembers = *PreparedServiceChanges->MaxMembers;
+	}
+
+	// Check for join policy change.
+	if (PreparedServiceChanges->JoinPolicy)
+	{
+		InternalPublicData->JoinPolicy = *PreparedServiceChanges->JoinPolicy;
+	}
 
 	// If no local members remain in the lobby process lobby removal.
 	if (bDispatchNotifications && LocalMembers.IsEmpty())
 	{
-		TArray<TSharedRef<FClientLobbyMemberSnapshot>> LeavingMembers;
-		MemberDataStorage.GenerateValueArray(LeavingMembers);
+		TArray<TSharedRef<FLobbyMemberInternal>> LeavingRemoteMembers;
+		MemberDataStorage.GenerateValueArray(LeavingRemoteMembers);
 
-		for (TSharedRef<FClientLobbyMemberSnapshot>& LeavingMember : LeavingMembers)
+		for (TSharedRef<FLobbyMemberInternal>& LeavingRemoteMember : LeavingRemoteMembers)
 		{
-			MemberDataStorage.Remove(LeavingMember->AccountId);
-			PublicData->Members.Remove(LeavingMember->AccountId);
-			LobbyEvents->OnLobbyMemberLeft.Broadcast(FLobbyMemberLeft{PublicData, LeavingMember, ELobbyMemberLeaveReason::Left});
+			RemoveMember(LeavingRemoteMember);
+			Params.LobbyEvents->OnLobbyMemberLeft.Broadcast(FLobbyMemberLeft{ InternalPublicData, LeavingRemoteMember, ELobbyMemberLeaveReason::Left });
 		}
 
-		LobbyEvents->OnLobbyLeft.Broadcast(FLobbyLeft{PublicData});
+		Params.LobbyEvents->OnLobbyLeft.Broadcast(FLobbyLeft{ InternalPublicData });
 	}
 
-	return Result;
+	// Clear out update.
+	ResetPreparedChanges();
+	return CommitResult;
 }
 
-FApplyLobbyUpdateResult FClientLobbyData::ApplyLobbyUpdateFromLocalChanges(FClientLobbyDataChanges&& Changes, FLobbyEvents& LobbyEvents)
+TOnlineResult<FLobbyClientDataPrepareClientChanges> FLobbyClientData::PrepareClientChanges(FLobbyClientDataPrepareClientChanges::Params&& Params) const
 {
-	FApplyLobbyUpdateResult Result;
-	const bool bDispatchLobbyNotifications = !LocalMembers.IsEmpty();
+	FPreparedClientChanges NewPreparedClientChanges;
+	FLobbyClientDataPrepareClientChanges::Result PrepareResult;
+
+	// Reset any pending changes. No other changes can occur while applying a service snapshot.
+	ResetPreparedChanges();
+
+	// The lobby schema may be changed by applying the update. Track the schema that is switched to so that the member attributes use the correct schema.
+	FSchemaId LobbySchemaId = InternalPublicData->SchemaId;
+
+	NewPreparedClientChanges.LocalAccountId = Params.LocalAccountId;
 
-	// Process lobby changes.
+	// Prepare lobby changes.
 	{
-		if (Changes.LocalName)
+		// Only the lobby owner may change some lobby data through local changes.
+		bool bLobbyDataChanged = false;
+
+		// Check for local name change.
+		if (Params.ClientChanges.LocalName && InternalPublicData->LocalName != *Params.ClientChanges.LocalName)
 		{
-			PublicData->LocalName = *Changes.LocalName;
+			if (InternalPublicData->LocalName != FName())
+			{
+				UE_LOG(LogTemp, Warning, TEXT("[FLobbyClientData::PrepareClientChanges] Local name chage failed - local name cannot be changed once set: Lobby[%s], CurrentName[%s], NewName[%s]"),
+					*ToLogString(InternalPublicData->LobbyId), *InternalPublicData->LocalName.ToString().ToLower(), *Params.ClientChanges.LocalName->ToString().ToLower());
+				return TOnlineResult<FLobbyClientDataPrepareClientChanges>(Errors::InvalidParams());
+			}
+
+			NewPreparedClientChanges.LocalName = *Params.ClientChanges.LocalName;
 		}
 
-		// Join policy.
-		if (Changes.JoinPolicy)
+		// Check for join policy change.
+		if (Params.ClientChanges.JoinPolicy && InternalPublicData->JoinPolicy != Params.ClientChanges.JoinPolicy)
 		{
-			PublicData->JoinPolicy = *Changes.JoinPolicy;
+			bLobbyDataChanged = true;
+			NewPreparedClientChanges.JoinPolicy = *Params.ClientChanges.JoinPolicy;
+			PrepareResult.ServiceChanges.JoinPolicy = *Params.ClientChanges.JoinPolicy;
 		}
 
-		// Lobby owner.
-		if (Changes.OwnerAccountId && *Changes.OwnerAccountId != PublicData->OwnerAccountId)
+		// Check for leader change.
+		if (Params.ClientChanges.OwnerAccountId && InternalPublicData->OwnerAccountId != Params.ClientChanges.OwnerAccountId)
 		{
-			if (TSharedRef<FClientLobbyMemberSnapshot>* MemberData = MemberDataStorage.Find(*Changes.OwnerAccountId))
+			if (!MemberDataStorage.Contains(*Params.ClientChanges.OwnerAccountId))
 			{
-				PublicData->OwnerAccountId = *Changes.OwnerAccountId;
-				if (bDispatchLobbyNotifications)
-				{
-					LobbyEvents.OnLobbyLeaderChanged.Broadcast(FLobbyLeaderChanged{PublicData, (*MemberData)});
-				}
+				UE_LOG(LogTemp, Warning, TEXT("[FLobbyClientData::PrepareClientChanges] Lobby owner chage failed - new owner is not in lobby: Lobby[%s], NewOwner[%s]"),
+					*ToLogString(InternalPublicData->LobbyId), *ToLogString(*Params.ClientChanges.OwnerAccountId));
+				return TOnlineResult<FLobbyClientDataPrepareClientChanges>(Errors::InvalidParams());
 			}
-			else
+
+			bLobbyDataChanged = true;
+			NewPreparedClientChanges.OwnerAccountId = *Params.ClientChanges.OwnerAccountId;
+			PrepareResult.ServiceChanges.OwnerAccountId = *Params.ClientChanges.OwnerAccountId;
+		}
+
+		// Check for user to kick.
+		if (Params.ClientChanges.KickedTargetMember)
+		{
+			if (!MemberDataStorage.Contains(*Params.ClientChanges.KickedTargetMember))
 			{
-				UE_LOG(LogTemp, Warning, TEXT("[FClientLobbyData::ApplyLobbyUpdateFromLocalChanges] Lobby owner chage failed - lobby member data not found: Lobby: %s, Member: %s"),
-					*ToLogString(PublicData->LobbyId), *ToLogString(*Changes.OwnerAccountId));
+				UE_LOG(LogTemp, Warning, TEXT("[FLobbyClientData::PrepareClientChanges] Kick target member failed - target is not in the lobby: Lobby[%s], KickedTargetMember[%s]"),
+					*ToLogString(InternalPublicData->LobbyId), *ToLogString(*Params.ClientChanges.KickedTargetMember));
+				return TOnlineResult<FLobbyClientDataPrepareClientChanges>(Errors::InvalidParams());
 			}
+
+			bLobbyDataChanged = true;
+			NewPreparedClientChanges.KickedTargetMember = *Params.ClientChanges.KickedTargetMember;
+			PrepareResult.ServiceChanges.KickedTargetMember = *Params.ClientChanges.KickedTargetMember;
 		}
 
-		// Lobby schema.
-		if (Changes.LobbySchema && *Changes.LobbySchema != PublicData->SchemaName)
+		// Check for attribute changes.
+		if (Params.ClientChanges.LobbySchema || Params.ClientChanges.Attributes)
 		{
-			PublicData->SchemaName = *Changes.LobbySchema;
-			if (bDispatchLobbyNotifications)
+			bLobbyDataChanged = true;
+
+			// Prepare lobby attribute changes.
+			FSchemaCategoryInstancePrepareClientChanges::Params PrepareClientChangesParams;
+			if (Params.ClientChanges.LobbySchema)
 			{
-				LobbyEvents.OnLobbySchemaChanged.Broadcast(FLobbySchemaChanged{PublicData});
+				PrepareClientChangesParams.ClientChanges.SchemaId = *Params.ClientChanges.LobbySchema;
+
+				// Update to using the new schema when applying a member update below.
+				LobbySchemaId = *Params.ClientChanges.LobbySchema;
 			}
+			if (Params.ClientChanges.Attributes)
+			{
+				PrepareClientChangesParams.ClientChanges.UpdatedAttributes = MoveTemp(Params.ClientChanges.Attributes->UpdatedAttributes);
+				PrepareClientChangesParams.ClientChanges.RemovedAttributes = MoveTemp(Params.ClientChanges.Attributes->RemovedAttributes);
+			}
+
+			TOnlineResult<FSchemaCategoryInstancePrepareClientChanges> PrepareClientChangesResult =
+				InternalPublicData->SchemaData.PrepareClientChanges(MoveTemp(PrepareClientChangesParams));
+			if (PrepareClientChangesResult.IsError())
+			{
+				UE_LOG(LogTemp, Warning, TEXT("[FLobbyClientData::PrepareClientChanges] Failed to prepare lobby attribute changes: Lobby[%s], Result[%s]"),
+					*ToLogString(InternalPublicData->LobbyId), *PrepareClientChangesResult.GetErrorValue().GetLogString());
+				return TOnlineResult<FLobbyClientDataPrepareClientChanges>(MoveTemp(PrepareClientChangesResult.GetErrorValue()));
+			}
+
+			// Get attribute changes to be written to the service.
+			PrepareResult.ServiceChanges.UpdatedAttributes = MoveTemp(PrepareClientChangesResult.GetOkValue().ServiceChanges.UpdatedAttributes);
+			PrepareResult.ServiceChanges.RemovedAttributes = MoveTemp(PrepareClientChangesResult.GetOkValue().ServiceChanges.RemovedAttributes);
+
+			NewPreparedClientChanges.bLobbyAttributesChanged = true;
 		}
 
-		// Apply lobby attribute changes.
-		if (!Changes.MutatedAttributes.IsEmpty() || !Changes.ClearedAttributes.IsEmpty())
+		if (bLobbyDataChanged)
 		{
-			TSet<FLobbyAttributeId> ChangedLobbyAttributes = ApplyAttributeUpdateFromChanges(
-				MoveTemp(Changes.MutatedAttributes),
-				MoveTemp(Changes.ClearedAttributes),
-				PublicData->Attributes);
+			// Check that the local user can make lobby changes.
+			if (Params.LocalAccountId != InternalPublicData->OwnerAccountId)
+			{
+				UE_LOG(LogTemp, Warning, TEXT("[FLobbyClientData::PrepareClientChanges] Failed to prepare lobby changes. LocalAccountId is not the lobby owner: Lobby[%s], LocalAccountId[%s], OwnerAccountId[%s]"),
+					*ToLogString(InternalPublicData->LobbyId), *ToLogString(Params.LocalAccountId), *ToLogString(InternalPublicData->OwnerAccountId));
+				return TOnlineResult<FLobbyClientDataPrepareClientChanges>(Errors::InvalidParams());
+			}
 
-			if (bDispatchLobbyNotifications && !ChangedLobbyAttributes.IsEmpty())
+			// Check that the local user is not trying to change the lobby while leaving.
+			if (Params.ClientChanges.LocalUserLeaveReason)
 			{
-				LobbyEvents.OnLobbyAttributesChanged.Broadcast(FLobbyAttributesChanged{PublicData, ChangedLobbyAttributes});
+				UE_LOG(LogTemp, Warning, TEXT("[FLobbyClientData::PrepareClientChanges] Failed to prepare lobby changes. Unable to modify the lobby while leaving: Lobby[%s], LocalAccountId[%s]"),
+					*ToLogString(InternalPublicData->LobbyId), *ToLogString(Params.LocalAccountId));
+				return TOnlineResult<FLobbyClientDataPrepareClientChanges>(Errors::InvalidParams());
 			}
 		}
 	}
 
-	// Apply member changes.
+	// Prepare lobby member changes.
 	{
-		// Adding local members will make the lobby active. When the lobby moves from inactive to
-		// active join events must be generated for the lobby and each of its existing members.
-		if (!Changes.MutatedMembers.IsEmpty())
+		// Check that the local user is not joining and leaving.
+		if (Params.ClientChanges.MemberAttributes && Params.ClientChanges.LocalUserLeaveReason)
 		{
-			// When adding members for the first time, notify for lobby joined then notify lobby member joined for each existing member.
-			if (LocalMembers.IsEmpty())
-			{
-				LobbyEvents.OnLobbyJoined.Broadcast(FLobbyJoined{PublicData});
+			UE_LOG(LogTemp, Warning, TEXT("[FLobbyClientData::PrepareClientChanges] Failed to prepare lobby member changes. Member cannot be updated and removed: Lobby[%s], LocalAccountId[%s]"),
+				*ToLogString(InternalPublicData->LobbyId), *ToLogString(Params.LocalAccountId));
+			return TOnlineResult<FLobbyClientDataPrepareClientChanges>(Errors::InvalidParams());
+		}
 
-				for (TPair<FAccountId, TSharedRef<FClientLobbyMemberSnapshot>>& MemberData : MemberDataStorage)
-				{
-					// The lobby creator will already be in the list during lobby creation due to applying a lobby snapshot.
-					if (Changes.MutatedMembers.Find(MemberData.Key) == nullptr)
-					{
-						LobbyEvents.OnLobbyMemberJoined.Broadcast(FLobbyMemberJoined{PublicData, MemberData.Value});
-					}
-				}
-			}
+		// Check whether the local user changed their attributes.
+		if (Params.ClientChanges.MemberAttributes)
+		{
+			FLobbyMemberInternal* PreparingMemberData = nullptr;
 
-			// Create or modify state for each mutated member and notify.
-			// The lobby creator will already be in the list during lobby creation due to applying a lobby snapshot.
-			for (TPair<FAccountId, TSharedRef<FClientLobbyMemberDataChanges>>& MutatedMemberPair : Changes.MutatedMembers)
+			// Check if this is a new or existing member.
+			if (const TSharedRef<FLobbyMemberInternal>* MemberData = MemberDataStorage.Find(Params.LocalAccountId))
 			{
-				if (TSharedRef<FClientLobbyMemberSnapshot>* MemberDataPtr = MemberDataStorage.Find(MutatedMemberPair.Key))
-				{
-					const bool bIsJoinedLocalMember = LocalMembers.Find(MutatedMemberPair.Key) != nullptr;
-
-					TSharedRef<FClientLobbyMemberSnapshot>& MemberData = *MemberDataPtr;
-					TSet<FLobbyAttributeId> ChangedAttributes = ApplyAttributeUpdateFromChanges(
-						MoveTemp(MutatedMemberPair.Value->MutatedAttributes),
-						MoveTemp(MutatedMemberPair.Value->ClearedAttributes),
-						MemberData->Attributes);
-
-					if (bIsJoinedLocalMember && !ChangedAttributes.IsEmpty())
-					{
-						LobbyEvents.OnLobbyMemberAttributesChanged.Broadcast(FLobbyMemberAttributesChanged{PublicData, MemberData, ChangedAttributes});
-					}
+				NewPreparedClientChanges.UpdatedLocalMember = *MemberData;
+				PreparingMemberData = &MemberData->Get();
 
-					if (!bIsJoinedLocalMember)
-					{
-						MemberData->bIsLocalMember = true;
-						LocalMembers.Add(MemberData->AccountId);
-						LobbyEvents.OnLobbyMemberJoined.Broadcast(FLobbyMemberJoined{PublicData, MemberData});
-					}
-				}
-				else
+				// If the member data exists and has not been flagged as a local user, then they
+				// are joining after a service snapshot with them present has been applied.
+				if (!LocalMembers.Contains(Params.LocalAccountId))
 				{
-					TSharedRef<FClientLobbyMemberSnapshot> JoiningMember = MakeShared<FClientLobbyMemberSnapshot>();
-					JoiningMember->bIsLocalMember = true;
-					JoiningMember->AccountId = MutatedMemberPair.Key;
-					JoiningMember->Attributes = MoveTemp(MutatedMemberPair.Value->MutatedAttributes);
-
-					// todo:
-					//NewMember->PlatformAccountId;
-					//NewMember->PlatformDisplayName;
-
-					MemberDataStorage.Add(JoiningMember->AccountId, JoiningMember);
-					PublicData->Members.Add(JoiningMember->AccountId, JoiningMember);
-					LocalMembers.Add(JoiningMember->AccountId);
-					LobbyEvents.OnLobbyMemberJoined.Broadcast(FLobbyMemberJoined{PublicData, JoiningMember});
+					NewPreparedClientChanges.JoiningLocalMember = *MemberData;
 				}
 			}
-		}
-
-		// Handle leaving members.
-		// When the last local member leaves, leave events will be generated for each of the
-		// remaining lobby members and for the lobby itself.
-		if (!Changes.LeavingMembers.IsEmpty())
-		{
-			for (TPair<FAccountId, ELobbyMemberLeaveReason>& LeavingMemberPair : Changes.LeavingMembers)
+			else
 			{
-				// Member data is expected to be added through changes before being mutated.
-				TSharedRef<FClientLobbyMemberSnapshot>* MemberDataPtr = MemberDataStorage.Find(LeavingMemberPair.Key);
-				if (ensure(MemberDataPtr))
-				{
-					TSharedRef<FClientLobbyMemberSnapshot> MemberData = *MemberDataPtr;
-					if (MemberData->bIsLocalMember)
-					{
-						Result.LeavingLocalMembers.Add(LeavingMemberPair.Key);
-						LocalMembers.Remove(LeavingMemberPair.Key);
-					}
+				// New member, create an entry for them.
+				TSharedRef<FLobbyMemberInternal> NewMemberData = MakeShared<FLobbyMemberInternal>(SchemaRegistry);
+				NewMemberData->AccountId = Params.LocalAccountId;
+				NewPreparedClientChanges.JoiningLocalMember = NewMemberData;
+				PreparingMemberData = &NewMemberData.Get();
+			}
 
-					MemberDataStorage.Remove(LeavingMemberPair.Key);
-					PublicData->Members.Remove(LeavingMemberPair.Key);
-					LobbyEvents.OnLobbyMemberLeft.Broadcast(FLobbyMemberLeft{PublicData, MemberData, LeavingMemberPair.Value});
-				}
+			// Prepare lobby member attribute changes.
+			FSchemaCategoryInstancePrepareClientChanges::Params PrepareClientChangesParams;
+			if (LobbySchemaId != InternalPublicData->SchemaId)
+			{
+				PrepareClientChangesParams.ClientChanges.SchemaId = LobbySchemaId;
 			}
+			PrepareClientChangesParams.ClientChanges.UpdatedAttributes = MoveTemp(Params.ClientChanges.MemberAttributes->UpdatedAttributes);
+			PrepareClientChangesParams.ClientChanges.RemovedAttributes = MoveTemp(Params.ClientChanges.MemberAttributes->RemovedAttributes);
 
-			// If no local members remain in the lobby process lobby removal.
-			if (LocalMembers.IsEmpty())
+			TOnlineResult<FSchemaCategoryInstancePrepareClientChanges> PrepareClientChangesResult =
+				PreparingMemberData->SchemaData.PrepareClientChanges(MoveTemp(PrepareClientChangesParams));
+			if (PrepareClientChangesResult.IsError())
 			{
-				TArray<TSharedRef<FClientLobbyMemberSnapshot>> LeavingMembers;
-				MemberDataStorage.GenerateValueArray(LeavingMembers);
+				UE_LOG(LogTemp, Warning, TEXT("[FLobbyClientData::PrepareClientChanges] Failed to prepare lobby member attribute changes: Lobby[%s], LocalMember[%s], Result[%s]"),
+					*ToLogString(InternalPublicData->LobbyId), *ToLogString(Params.LocalAccountId), *PrepareClientChangesResult.GetErrorValue().GetLogString());
+				return TOnlineResult<FLobbyClientDataPrepareClientChanges>(MoveTemp(PrepareClientChangesResult.GetErrorValue()));
+			}
 
-				for (TSharedRef<FClientLobbyMemberSnapshot>& LeavingMember : LeavingMembers)
-				{
-					MemberDataStorage.Remove(LeavingMember->AccountId);
-					PublicData->Members.Remove(LeavingMember->AccountId);
-					LobbyEvents.OnLobbyMemberLeft.Broadcast(FLobbyMemberLeft{PublicData, LeavingMember, ELobbyMemberLeaveReason::Left});
-				}
+			// Get attribute changes to be written to the service.
+			PrepareResult.ServiceChanges.UpdatedMemberAttributes = MoveTemp(PrepareClientChangesResult.GetOkValue().ServiceChanges.UpdatedAttributes);
+			PrepareResult.ServiceChanges.RemovedMemberAttributes = MoveTemp(PrepareClientChangesResult.GetOkValue().ServiceChanges.RemovedAttributes);
+		}
 
-				LobbyEvents.OnLobbyLeft.Broadcast(FLobbyLeft{PublicData});
+		// Check whether the local user left.
+		if (Params.ClientChanges.LocalUserLeaveReason)
+		{
+			if (!MemberDataStorage.Contains(Params.LocalAccountId))
+			{
+				UE_LOG(LogTemp, Warning, TEXT("[FLobbyClientData::PrepareClientChanges] Local user leave faild - user is not in the lobby: Lobby[%s], LocalAccountId[%s]"),
+					*ToLogString(InternalPublicData->LobbyId), *ToLogString(Params.LocalAccountId));
+				return TOnlineResult<FLobbyClientDataPrepareClientChanges>(Errors::InvalidParams());
 			}
+
+			NewPreparedClientChanges.LocalUserLeaveReason = *Params.ClientChanges.LocalUserLeaveReason;
 		}
 	}
 
-	return Result;
+	PreparedClientChanges = NewPreparedClientChanges;
+	return TOnlineResult<FLobbyClientDataPrepareClientChanges>(MoveTemp(PrepareResult));
 }
 
-TSet<FLobbyAttributeId> FClientLobbyData::ApplyAttributeUpdateFromSnapshot(
-	TMap<FLobbyAttributeId, FLobbyVariant>&& AttributeSnapshot,
-	TMap<FLobbyAttributeId, FLobbyVariant>& ExistingAttributes)
+FLobbyClientDataCommitClientChanges::Result FLobbyClientData::CommitClientChanges(FLobbyClientDataCommitClientChanges::Params&& Params)
 {
-	TSet<FLobbyAttributeId> Result;
+	FLobbyClientDataCommitClientChanges::Result CommitResult;
+	bool bDispatchNotifications = Params.LobbyEvents != nullptr && !LocalMembers.IsEmpty();
 
-	// Clear the attributes.
+	if (!PreparedClientChanges)
 	{
-		// Determine which attributes were removed.
-		for (const TPair<FLobbyAttributeId, FLobbyVariant>& Attribute : ExistingAttributes)
+		UE_LOG(LogTemp, Warning, TEXT("[FLobbyClientData::CommitClientChanges] Commit failed, no pending changes found: Lobby[%s]"),
+			*ToLogString(InternalPublicData->LobbyId));
+		return CommitResult;
+	}
+
+	// Commit lobby changes.
+	{
+		// Assign local name.
+		if (PreparedClientChanges->LocalName)
+		{
+			InternalPublicData->LocalName = *PreparedClientChanges->LocalName;
+		}
+
+		// Assign join policy.
+		if (PreparedClientChanges->JoinPolicy)
+		{
+			InternalPublicData->JoinPolicy = *PreparedClientChanges->JoinPolicy;
+		}
+
+		// Assign join policy
+		if (PreparedClientChanges->OwnerAccountId)
 		{
-			if (!AttributeSnapshot.Contains(Attribute.Key))
+			InternalPublicData->OwnerAccountId = *PreparedClientChanges->OwnerAccountId;
+
+			if (bDispatchNotifications)
 			{
-				Result.Add(Attribute.Key);
+				TSharedRef<FLobbyMemberInternal> MemberData = MemberDataStorage.FindChecked(InternalPublicData->OwnerAccountId);
+				Params.LobbyEvents->OnLobbyLeaderChanged.Broadcast({ InternalPublicData, MemberData });
 			}
 		}
 
-		// Remove them.
-		for (const FLobbyAttributeId& AttributeId : Result)
+		// Remove kicked member.
+		if (PreparedClientChanges->KickedTargetMember)
 		{
-			ExistingAttributes.Remove(AttributeId);
+			TSharedRef<FLobbyMemberInternal> MemberData = MemberDataStorage.FindChecked(*PreparedClientChanges->KickedTargetMember);
+			RemoveMember(MemberData);
+
+			if (bDispatchNotifications)
+			{
+				Params.LobbyEvents->OnLobbyMemberLeft.Broadcast({ InternalPublicData, MemberData, ELobbyMemberLeaveReason::Kicked });
+			}
 		}
-	}
 
-	// Apply attribute additions and mutations.
-	{
-		for (TPair<FLobbyAttributeId, FLobbyVariant>& Attribute : AttributeSnapshot)
+		// Commit schema changes if needed.
+		if (PreparedClientChanges->bLobbyAttributesChanged)
 		{
-			if (FLobbyVariant* ExistingAttributeData = ExistingAttributes.Find(Attribute.Key))
+			FSchemaCategoryInstanceCommitClientChanges::Result CommitLobbyServiceSnapshotResult = InternalPublicData->SchemaData.CommitClientChanges();
+			FSchemaServiceClientChanges& LobbyServiceClientChanges = CommitLobbyServiceSnapshotResult.ClientChanges;
+
+			// Notify if schema changed.
+			// Schema change notification must be processed first as the schema affects how the client processes attributes.
+			if (LobbyServiceClientChanges.SchemaId)
 			{
-				if (*ExistingAttributeData != Attribute.Value)
+				FLobbySchemaChanged LobbySchemaChanged = { InternalPublicData, InternalPublicData->SchemaId, *LobbyServiceClientChanges.SchemaId };
+				InternalPublicData->SchemaId = *LobbyServiceClientChanges.SchemaId;
+				if (bDispatchNotifications)
 				{
-					Result.Add(MoveTemp(Attribute.Key));
-					*ExistingAttributeData = MoveTemp(Attribute.Value);
+					Params.LobbyEvents->OnLobbySchemaChanged.Broadcast(LobbySchemaChanged);
 				}
 			}
-			else
+
+			const bool bAnyLobbyAttributesChanged =
+				!LobbyServiceClientChanges.AddedAttributes.IsEmpty() ||
+				!LobbyServiceClientChanges.ChangedAttributes.IsEmpty() ||
+				!LobbyServiceClientChanges.RemovedAttributes.IsEmpty();
+
+			// Handle lobby attribute change notifications.
+			if (bDispatchNotifications && bAnyLobbyAttributesChanged)
 			{
-				Result.Add(Attribute.Key);
-				ExistingAttributes.Add(MoveTemp(Attribute));
+				Params.LobbyEvents->OnLobbyAttributesChanged.Broadcast({
+					InternalPublicData,
+					MoveTemp(LobbyServiceClientChanges.AddedAttributes),
+					MoveTemp(LobbyServiceClientChanges.ChangedAttributes),
+					MoveTemp(LobbyServiceClientChanges.RemovedAttributes) });
 			}
 		}
 	}
 
-	return Result;
-}
+	// Check for an updating local member.
+	if (PreparedClientChanges->UpdatedLocalMember)
+	{
+		TSharedRef<FLobbyMemberInternal> UpdatedMember = PreparedClientChanges->UpdatedLocalMember.ToSharedRef();
 
-TSet<FLobbyAttributeId> FClientLobbyData::ApplyAttributeUpdateFromChanges(
-	TMap<FLobbyAttributeId, FLobbyVariant>&& MutatedAttributes,
-	TSet<FLobbyAttributeId>&& ClearedAttributes,
-	TMap<FLobbyAttributeId, FLobbyVariant>& ExistingAttributes)
-{
-	TSet<FLobbyAttributeId> Result;
+		// Commit attribute changes.
+		FSchemaCategoryInstanceCommitClientChanges::Result CommitLobbyMemberServiceSnapshotResult = UpdatedMember->SchemaData.CommitClientChanges();
+		FSchemaServiceClientChanges& LobbyMemberServiceClientChanges = CommitLobbyMemberServiceSnapshotResult.ClientChanges;
+
+		const bool bAnyLobbyMemberAttributesChanged =
+			!LobbyMemberServiceClientChanges.AddedAttributes.IsEmpty() ||
+			!LobbyMemberServiceClientChanges.ChangedAttributes.IsEmpty() ||
+			!LobbyMemberServiceClientChanges.RemovedAttributes.IsEmpty();
+
+		// Handle lobby attribute changes.
+		if (bDispatchNotifications && bAnyLobbyMemberAttributesChanged)
+		{
+			Params.LobbyEvents->OnLobbyMemberAttributesChanged.Broadcast({
+				InternalPublicData,
+				UpdatedMember,
+				MoveTemp(LobbyMemberServiceClientChanges.AddedAttributes),
+				MoveTemp(LobbyMemberServiceClientChanges.ChangedAttributes),
+				MoveTemp(LobbyMemberServiceClientChanges.RemovedAttributes) });
+		}
+	}
 
-	// Apply attribute additions and mutations.
+	// Commit lobby member changes.
 	{
-		for (TPair<FLobbyAttributeId, FLobbyVariant>& Attribute : MutatedAttributes)
+		// Check for a new joining local member.
+		if (PreparedClientChanges->JoiningLocalMember)
 		{
-			if (FLobbyVariant* ExistingAttributeData = ExistingAttributes.Find(Attribute.Key))
+			// Check whether notifications should be turned on.
+			if (LocalMembers.IsEmpty())
 			{
-				if (*ExistingAttributeData != Attribute.Value)
+				bDispatchNotifications = Params.LobbyEvents != nullptr;
+				if (bDispatchNotifications)
 				{
-					Result.Add(MoveTemp(Attribute.Key));
-					*ExistingAttributeData = MoveTemp(Attribute.Value);
+					// Signal that the lobby has been joined.
+					Params.LobbyEvents->OnLobbyJoined.Broadcast({ InternalPublicData });
+
+					// Dispatch a notification for each existing remote member.
+					for (TPair<FAccountId, TSharedRef<FLobbyMemberInternal>>& MemberPair : MemberDataStorage)
+					{
+						Params.LobbyEvents->OnLobbyMemberJoined.Broadcast({ InternalPublicData, MemberPair.Value });
+					}
 				}
 			}
-			else
+
+			TSharedRef<FLobbyMemberInternal> JoinedMember = PreparedClientChanges->JoiningLocalMember.ToSharedRef();
+
+			JoinedMember->bIsLocalMember = true;
+			// todo.
+			//NewMemberData->PlatformAccountId = ;
+			//NewMemberData->PlatformDisplayName = ;
+
+			const bool bWasAlreadyJoined = MemberDataStorage.Contains(JoinedMember->AccountId);
+
+			// Add to member tracking.
+			AddMember(JoinedMember);
+
+			// Commit attribute changes.
+			JoinedMember->SchemaData.CommitServiceSnapshot();
+
+			// Send join notification for the local user.
+			if (bDispatchNotifications && !bWasAlreadyJoined)
 			{
-				Result.Add(Attribute.Key);
-				ExistingAttributes.Add(MoveTemp(Attribute));
+				Params.LobbyEvents->OnLobbyMemberJoined.Broadcast({ InternalPublicData, JoinedMember });
 			}
 		}
-	}
 
-	// Clear the attributes.
-	{
-		for (FLobbyAttributeId ClearedAttributeId : ClearedAttributes)
+		// Check if the local user left.
+		if (PreparedClientChanges->LocalUserLeaveReason)
 		{
-			if (ExistingAttributes.Remove(ClearedAttributeId) > 0)
+			TSharedRef<FLobbyMemberInternal> LeavingMember = MemberDataStorage.FindChecked(PreparedClientChanges->LocalAccountId);
+			RemoveMember(LeavingMember);
+			CommitResult.LeavingLocalMembers.Add(PreparedClientChanges->LocalAccountId);
+
+			if (bDispatchNotifications)
 			{
-				Result.Add(ClearedAttributeId);
+				Params.LobbyEvents->OnLobbyMemberLeft.Broadcast({ InternalPublicData, LeavingMember, *PreparedClientChanges->LocalUserLeaveReason });
+
+				if (LocalMembers.IsEmpty())
+				{
+					// Fire leave events for remote members.
+					TArray<TSharedRef<FLobbyMemberInternal>> LeavingRemoteMembers;
+					MemberDataStorage.GenerateValueArray(LeavingRemoteMembers);
+
+					for (TSharedRef<FLobbyMemberInternal>& LeavingRemoteMember : LeavingRemoteMembers)
+					{
+						RemoveMember(LeavingMember);
+						Params.LobbyEvents->OnLobbyMemberLeft.Broadcast(FLobbyMemberLeft{ InternalPublicData, LeavingRemoteMember, ELobbyMemberLeaveReason::Left });
+					}
+
+					// Notify lobby has been left.
+					Params.LobbyEvents->OnLobbyLeft.Broadcast({ InternalPublicData });
+				}
 			}
 		}
 	}
 
-	return Result;
+	ResetPreparedChanges();
+	return CommitResult;
+}
+
+void FLobbyClientData::ResetPreparedChanges() const
+{
+	PreparedServiceChanges.Reset();
+	PreparedClientChanges.Reset();
+}
+
+void FLobbyClientData::AddMember(const TSharedRef<FLobbyMemberInternal>& Member)
+{
+	MemberDataStorage.Add(Member->AccountId, Member);
+	InternalPublicData->Members.Add(Member->AccountId, Member);
+
+	if (Member->bIsLocalMember)
+	{
+		LocalMembers.Add(Member->AccountId);
+	}
+}
+
+void FLobbyClientData::RemoveMember(const TSharedRef<FLobbyMemberInternal>& Member)
+{
+	MemberDataStorage.Remove(Member->AccountId);
+	InternalPublicData->Members.Remove(Member->AccountId);
+
+	if (Member->bIsLocalMember)
+	{
+		LocalMembers.Remove(Member->AccountId);
+	}
 }
 
 /* UE::Online */ }
diff --git a/Engine/Plugins/Online/Experimental/OnlineServices/Source/OnlineServicesCommon/Public/Online/LobbiesCommon.h b/Engine/Plugins/Online/Experimental/OnlineServices/Source/OnlineServicesCommon/Public/Online/LobbiesCommon.h
--- a/Engine/Plugins/Online/Experimental/OnlineServices/Source/OnlineServicesCommon/Public/Online/LobbiesCommon.h
+++ b/Engine/Plugins/Online/Experimental/OnlineServices/Source/OnlineServicesCommon/Public/Online/LobbiesCommon.h
@@ -55,9 +55,8 @@ protected:
 	TOnlineAsyncOpHandle<FFunctionalTestLobbies> FunctionalTest(FFunctionalTestLobbies::Params&& Params);
 #endif // LOBBIES_FUNCTIONAL_TEST_ENABLED
 
+	TSharedRef<FSchemaRegistry> SchemaRegistry;
 	FLobbyEvents LobbyEvents;
-	TSharedPtr<FLobbySchemaRegistry> LobbySchemaRegistry;
-	TSharedPtr<FLobbySchema> ServiceSchema;
 };
 
 /* UE::Online */ }
diff --git a/Engine/Plugins/Online/Experimental/OnlineServices/Source/OnlineServicesCommon/Public/Online/LobbiesCommonTypes.h b/Engine/Plugins/Online/Experimental/OnlineServices/Source/OnlineServicesCommon/Public/Online/LobbiesCommonTypes.h
--- a/Engine/Plugins/Online/Experimental/OnlineServices/Source/OnlineServicesCommon/Public/Online/LobbiesCommonTypes.h
+++ b/Engine/Plugins/Online/Experimental/OnlineServices/Source/OnlineServicesCommon/Public/Online/LobbiesCommonTypes.h
@@ -4,94 +4,17 @@
 
 #include "Async/Future.h"
 #include "Online/Lobbies.h"
+#include "Online/Schema.h"
 
 namespace UE::Online {
 
 #define LOBBIES_FUNCTIONAL_TEST_ENABLED !UE_BUILD_SHIPPING
 
-class FOnlineServicesCommon;
-
-struct FLobbyConfig
-{
-	TArray<FLobbySchemaId> RegisteredSchemas;
-};
-
-struct FLobbySchemaAttributeConfig
-{
-	// The attribute name.
-	FLobbyAttributeId Name;
-	
-	// Override default service attribute field assignment.
-	// Specifying an override allows custom grouping of attributes within a
-	// platform attribute. This grouping is useful when a set of attributes are
-	// known to update at the same cadence.
-	TOptional<FLobbyAttributeId> ServiceAttributeFieldId;
-	
-	// The size in bytes of the attribute value.
-	// A schema will fail validation if all attributes will not fit within the
-	// platform's lobby fields.
-	uint32 MaxByteSize = 0;
-
-	// Control whether the attribute is visible to players who are not joined
-	// to the lobby.
-	ELobbyAttributeVisibility Visibility = ELobbyAttributeVisibility::Private;
-
-};
-
-struct FLobbySchemaConfig
-{
-	// The schema name.
-	FLobbySchemaId SchemaName;
+ONLINESERVICESCOMMON_API extern FSchemaId LobbyBaseSchemaId;
+ONLINESERVICESCOMMON_API extern FSchemaCategoryId LobbySchemaCategoryId;
+ONLINESERVICESCOMMON_API extern FSchemaCategoryId LobbyMemberSchemaCategoryId;
 
-	// The optional base schema name.
-	FLobbySchemaId BaseSchemaName;
-	
-	// The definitions for attributes attached to a lobby.
-	TArray<FLobbySchemaAttributeConfig> LobbyAttributes;
-
-	// The definitions for attributes attached to a lobby member.
-	TArray<FLobbySchemaAttributeConfig> LobbyMemberAttributes;
-};
-
-class FLobbyServiceAttributeChanges final
-{
-public:
-	TMap<FLobbyAttributeId, FLobbyVariant> MutatedAttributes;
-	TSet<FLobbyAttributeId> ClearedAttributes;
-};
-
-class FLobbyClientAttributeChanges final
-{
-public:
-	// Apply attribute changes to attribute parameter.
-	// Applying the changes will clear MutatedAttributes and ClearedAttributes.
-	// Returns a list of the attributes which were changed.
-	TSet<FLobbyAttributeId> Apply(TMap<FLobbyAttributeId, FLobbyVariant>& InOutAttributes)
-	{
-		TSet<FLobbyAttributeId> Changes;
-		Changes.Reserve(MutatedAttributes.Num() + ClearedAttributes.Num());
-
-		for (TPair<FLobbyAttributeId, FLobbyVariant>& Attribute : MutatedAttributes)
-		{
-			Changes.Add(Attribute.Key);
-			InOutAttributes.Emplace(Attribute.Key, MoveTemp(Attribute.Value));
-		}
-
-		for (FLobbyAttributeId AttributeId : ClearedAttributes)
-		{
-			Changes.Add(AttributeId);
-			InOutAttributes.Remove(AttributeId);
-		}
-
-		MutatedAttributes.Empty();
-		ClearedAttributes.Empty();
-
-		return Changes;
-	}
-
-	TMap<FLobbyAttributeId, FLobbyVariant> MutatedAttributes;
-	TSet<FLobbyAttributeId> ClearedAttributes;
-};
+class FOnlineServicesCommon;
 
 struct FLobbyEvents final
 {
@@ -108,32 +31,6 @@ struct FLobbyEvents final
 	TOnlineEventCallable<void(const FUILobbyJoinRequested&)> OnUILobbyJoinRequested;
 };
 
-class ONLINESERVICESCOMMON_API FLobbySchema final
-{
-public:
-	static TSharedPtr<FLobbySchema> Create(FLobbySchemaConfig LobbySchemaConfig);
-
-	TSharedRef<FLobbyServiceAttributeChanges> TranslateLobbyAttributes(const TMap<FLobbyAttributeId, FLobbyVariant>& LobbyAttributes, const FLobbyClientAttributeChanges& ClientAttributeChanges) const;
-	TSharedRef<FLobbyClientAttributeChanges> TranslateLobbyAttributes(const FLobbyServiceAttributeChanges& ServiceAttributeChanges) const;
-	TSharedRef<FLobbyServiceAttributeChanges> TranslateLobbyMemberAttributes(const TMap<FLobbyAttributeId, FLobbyVariant>& LobbyAttributes, const FLobbyClientAttributeChanges& ClientAttributeChanges) const;
-	TSharedRef<FLobbyClientAttributeChanges> TranslateLobbyMemberAttributes(const FLobbyServiceAttributeChanges& ServiceAttributeChanges) const;
-
-private:
-	FLobbySchemaId SchemaId;
-};
-
-class FLobbySchemaRegistry
-{
-public:
-	bool Initialize(TArray<FLobbySchemaConfig> LobbySchemaConfigs);
-	TSharedPtr<FLobbySchema> FindSchema(FLobbySchemaId SchemaId);
-
-private:
-	bool RegisterSchema(FLobbySchemaConfig LobbySchemaConfig);
-
-	TMap<FLobbySchemaId, TSharedRef<FLobbySchema>> RegisteredSchemas;
-};
-
 // Todo: put this somewhere else
 template <typename AwaitedType>
 TFuture<TArray<AwaitedType>> WhenAll(TArray<TFuture<AwaitedType>>&& Futures)
@@ -172,127 +69,293 @@ TFuture<TArray<AwaitedType>> WhenAll(TArray<TFuture<AwaitedType>>&& Futures)
 }
 
 /**
- * Local changes to a member to be applied to an existing lobby member snapshot.
- */
-struct FClientLobbyMemberDataChanges
+* Accessor for using the public lobby attributes directly when committing schema changes.
+*/
+class FSchemaCategoryInstanceLobbySnapshotAccessor
+{
+public:
+	FSchemaCategoryInstanceLobbySnapshotAccessor(TMap<FSchemaAttributeId, FSchemaVariant>& InClientSnapshot)
+		: ClientSnapshot(InClientSnapshot)
+	{
+	}
+
+	TMap<FSchemaAttributeId, FSchemaVariant>& GetMutableClientSnapshot()
+	{
+		return ClientSnapshot;
+	}
+private:
+
+	TMap<FSchemaAttributeId, FSchemaVariant>& ClientSnapshot;
+};
+
+/**
+* Alias for schema category instance used by lobby and lobby member data.
+*/
+using FLobbiesSchemaCategoryInstance = TSchemaCategoryInstance<FSchemaCategoryInstanceLobbySnapshotAccessor>;
+
+/**
+* Lobby data structure as seen by the client with additional fields for internal use.
+*/
+struct FLobbyInternal : public FLobby
+{
+	FLobbyInternal(const TSharedRef<const FSchemaRegistry>& SchemaRegistry)
+		: FLobby()
+		, SchemaData(FSchemaId(), LobbyBaseSchemaId, LobbySchemaCategoryId, SchemaRegistry, Attributes)
+	{
+	}
+
+	FLobbiesSchemaCategoryInstance SchemaData;
+};
+
+/**
+* Lobby member data structure as seen by the client with additional fields for internal use.
+*/
+struct FLobbyMemberInternal : public FLobbyMember
+{
+	FLobbyMemberInternal(const TSharedRef<const FSchemaRegistry>& SchemaRegistry)
+		: FLobbyMember()
+		, SchemaData(FSchemaId(), LobbyBaseSchemaId, LobbyMemberSchemaCategoryId, SchemaRegistry, Attributes)
+	{
+	}
+
+	FLobbiesSchemaCategoryInstance SchemaData;
+};
+
+struct FLobbyClientAttributeChanges
 {
 	/** New or changed attributes. */
-	TMap<FLobbyAttributeId, FLobbyVariant> MutatedAttributes;
+	TMap<FSchemaAttributeId, FSchemaVariant> UpdatedAttributes;
+	/** Attributes to be removed. */
+	TSet<FSchemaAttributeId> RemovedAttributes;
+};
 
-	/** Attributes to be cleared. */
-	TSet<FLobbyAttributeId> ClearedAttributes;
+struct FLobbyClientMemberAttributeChanges
+{
+	/** New or changed attributes. */
+	TMap<FSchemaAttributeId, FSchemaVariant> UpdatedAttributes;
+	/** Attributes to be removed. */
+	TSet<FSchemaAttributeId> RemovedAttributes;
 };
 
 /**
- * Local changes to a lobby to be applied to an existing lobby snapshot.
+ * Local changes to a lobby to be applied to the service.
  */
-struct FClientLobbyDataChanges
+struct FLobbyClientChanges
 {
 	/** Local name for the lobby. */
 	TOptional<FName> LocalName;
-
 	/** Setting for new join policy. */
 	TOptional<ELobbyJoinPolicy> JoinPolicy;
-
 	/** Setting for lobby ownership change. */
 	TOptional<FAccountId> OwnerAccountId;
-
 	/** Setting for lobby schema change. */
-	TOptional<FLobbySchemaId> LobbySchema;
-
-	/** New or changed attributes. */
-	TMap<FLobbyAttributeId, FLobbyVariant> MutatedAttributes;
-
-	/** Attributes to be cleared. */
-	TSet<FLobbyAttributeId> ClearedAttributes;
-
-	/** Members to be added or changed. */
-	TMap<FAccountId, TSharedRef<FClientLobbyMemberDataChanges>> MutatedMembers;
+	TOptional<FSchemaId> LobbySchema;
+	/** Setting for changing lobby attributes. */
+	TOptional<FLobbyClientAttributeChanges> Attributes;
+	/** Setting for adding or updating the local user. */
+	TOptional<FLobbyClientMemberAttributeChanges> MemberAttributes;
+	/** Setting for removing the local user. */
+	TOptional<ELobbyMemberLeaveReason> LocalUserLeaveReason;
+	/** Setting for kicking the target lobby member. */
+	TOptional<FAccountId> KickedTargetMember;
+};
 
-	/** Members to be removed. */
-	TMap<FAccountId, ELobbyMemberLeaveReason> LeavingMembers;
+/**
+ * Local changes to a lobby translated to the changes which will be applied to the service.
+ */
+struct FLobbyClientServiceChanges
+{
+	/** Setting for new join policy. */
+	TOptional<ELobbyJoinPolicy> JoinPolicy;
+	/** Setting for lobby ownership change. */
+	TOptional<FAccountId> OwnerAccountId;
+	/** Target member to be kicked. */
+	TOptional<FAccountId> KickedTargetMember;
+	/** Added or changed lobby attributes and their values. */
+	TMap<FSchemaServiceAttributeId, FSchemaServiceAttributeData> UpdatedAttributes;
+	/** Removed lobby attribute ids. */
+	TSet<FSchemaServiceAttributeId> RemovedAttributes;
+	/** Added or changed lobby member attributes and their values. */
+	TMap<FSchemaServiceAttributeId, FSchemaServiceAttributeData> UpdatedMemberAttributes;
+	/** Removed lobby member attribute ids. */
+	TSet<FSchemaServiceAttributeId> RemovedMemberAttributes;
 };
 
 /**
-* Lobby snapshot data.
-* All contained data has been translated from EOS types.
-* Attributes have had their schema transformations applied.
+* Full snapshot of lobby data excluding the individual member attributes
 */
-struct FClientLobbySnapshot
+struct FLobbyServiceSnapshot
 {
 	FAccountId OwnerAccountId;
-	FName SchemaName;
 	int32 MaxMembers;
 	ELobbyJoinPolicy JoinPolicy;
-	TMap<FLobbyAttributeId, FLobbyVariant> Attributes;
+	FSchemaServiceSnapshot SchemaServiceSnapshot;
 	TSet<FAccountId> Members;
 };
 
-struct FClientLobbyMemberSnapshot : public FLobbyMember
+/**
+* Full snapshot of a lobby member.
+*/
+struct FLobbyMemberServiceSnapshot
 {
-	bool bIsLocalMember = false;
+	FAccountId AccountId;
+	FAccountId PlatformAccountId;
+	FString PlatformDisplayName;
+	FSchemaServiceSnapshot SchemaServiceSnapshot;
 };
 
-struct FApplyLobbyUpdateResult
+struct FLobbyClientDataPrepareServiceSnapshot
 {
-	TArray<FAccountId> LeavingLocalMembers;
+	struct Params
+	{
+		/** Full snapshot from the service. */
+		FLobbyServiceSnapshot LobbySnapshot;
+		/** Full snapshot for each member from the service. */
+		TMap<FAccountId, FLobbyMemberServiceSnapshot> LobbyMemberSnapshots;
+		/** Any leaving members notified by the service. */
+		TMap<FAccountId, ELobbyMemberLeaveReason> LeaveReasons;
+	};
+
+	struct Result
+	{
+	};
+};
+
+struct FLobbyClientDataCommitServiceSnapshot
+{
+	struct Params
+	{
+		/** Notifier for signaling client events on data change. */
+		FLobbyEvents* LobbyEvents = nullptr;
+	};
+
+	struct Result
+	{
+		/** List of local members leaving the lobby. */
+		TArray<FAccountId> LeavingLocalMembers;
+	};
+};
+
+struct FLobbyClientDataPrepareClientChanges
+{
+	struct Params
+	{
+		/** The local user who is making the changes. */
+		FAccountId LocalAccountId;
+		/** Changes to be applied to service data. */
+		FLobbyClientChanges ClientChanges;
+	};
+
+	struct Result
+	{
+		/** Translated changes to be applied to the service. */
+		FLobbyClientServiceChanges ServiceChanges;
+	};
+};
+
+struct FLobbyClientDataCommitClientChanges
+{
+	struct Params
+	{
+		/** Notifier for signaling client events on data change. */
+		FLobbyEvents* LobbyEvents = nullptr;
+	};
+
+	struct Result
+	{
+		/** List of local members leaving the lobby. */
+		TArray<FAccountId> LeavingLocalMembers;
+	};
 };
 
 /** Lobby data as seen by the client. */
-struct ONLINESERVICESCOMMON_API FClientLobbyData final
+struct ONLINESERVICESCOMMON_API FLobbyClientData final
 {
 public:
-	FClientLobbyData(FLobbyId LobbyId);
+	FLobbyClientData(FLobbyId LobbyId, const TSharedRef<const FSchemaRegistry>& InSchemaRegistry);
 
-	TSharedRef<const FLobby> GetPublicDataPtr() const { return PublicData; }
-	const FLobby& GetPublicData() const { return *PublicData; }
+	TSharedRef<const FLobby> GetPublicDataPtr() const { return InternalPublicData; }
+	const FLobby& GetPublicData() const { return *InternalPublicData; }
 
-	TSharedPtr<const FClientLobbyMemberSnapshot> GetMemberData(FAccountId MemberAccountId) const;
+	TSharedPtr<const FLobbyMemberInternal> GetMemberData(FAccountId MemberAccountId) const;
 
 	/**
 	 * Apply updated lobby data and generate changes.
 	 * The LeaveReason provides context for members who have left since the most recent snapshot.
 	 * Changing a lobby generates events for the local client.
 	 */
-	FApplyLobbyUpdateResult ApplyLobbyUpdateFromServiceSnapshot(
-		FClientLobbySnapshot&& LobbySnapshot,
-		TMap<FAccountId, TSharedRef<FClientLobbyMemberSnapshot>>&& LobbyMemberSnapshots,
-		TMap<FAccountId, ELobbyMemberLeaveReason>&& LeaveReasons = TMap<FAccountId, ELobbyMemberLeaveReason>(),
-		FLobbyEvents* LobbyEvents = nullptr);
+	TOnlineResult<FLobbyClientDataPrepareServiceSnapshot> PrepareServiceSnapshot(FLobbyClientDataPrepareServiceSnapshot::Params&& Params) const;
+	FLobbyClientDataCommitServiceSnapshot::Result CommitServiceSnapshot(FLobbyClientDataCommitServiceSnapshot::Params&& Params);
 
-	/**
-	 * Apply updated lobby data and generate changes.
-	 * Changing a lobby generates events for the local client.
-	 */
-	FApplyLobbyUpdateResult ApplyLobbyUpdateFromLocalChanges(FClientLobbyDataChanges&& Changes, FLobbyEvents& LobbyEvents);
+	TOnlineResult<FLobbyClientDataPrepareClientChanges> PrepareClientChanges(FLobbyClientDataPrepareClientChanges::Params&& Params) const;
+	FLobbyClientDataCommitClientChanges::Result CommitClientChanges(FLobbyClientDataCommitClientChanges::Params&& Params);
 
 private:
+	void ResetPreparedChanges() const;
+	void AddMember(const TSharedRef<FLobbyMemberInternal>& Member);
+	void RemoveMember(const TSharedRef<FLobbyMemberInternal>& Member);
 
-	/** Apply changes to a set of attributes. Returns the set of attribute IDs which changed. */
-	TSet<FLobbyAttributeId> ApplyAttributeUpdateFromSnapshot(
-		TMap<FLobbyAttributeId, FLobbyVariant>&& AttributeSnapshot,
-		TMap<FLobbyAttributeId, FLobbyVariant>& ExistingAttributes);
+	struct FPreparedServiceChanges
+	{
+		/** Set when there is a change in lobby ownership. */
+		TOptional<FAccountId> OwnerAccountId;
+		/** Set when there is a change in max members. */
+		TOptional<int32> MaxMembers;
+		/** Set when there is a change in join policy. */
+		TOptional<ELobbyJoinPolicy> JoinPolicy;
+		/** Members who did not already have a local representation. */
+		TArray<TSharedRef<FLobbyMemberInternal>> NewMemberDataStorage;
+		/** Members who already exist in the lobby. */
+		TArray<TSharedRef<FLobbyMemberInternal>> UpdatedMemberDataStorage;
+		/** Any leaving members and their reason. */
+		TArray<TPair<TSharedRef<FLobbyMemberInternal>, ELobbyMemberLeaveReason>> LeavingMembers;
+	};
 
-	/** Apply changes to a set of attributes. Returns the set of attribute IDs which changed. */
-	TSet<FLobbyAttributeId> ApplyAttributeUpdateFromChanges(
-		TMap<FLobbyAttributeId, FLobbyVariant>&& MutatedAttributes,
-		TSet<FLobbyAttributeId>&& ClearedAttributes,
-		TMap<FLobbyAttributeId, FLobbyVariant>& ExistingAttributes);
+	struct FPreparedClientChanges
+	{
+		/** The local user associated with the change. */
+		FAccountId LocalAccountId;
+		/** Local name for the lobby. */
+		TOptional<FName> LocalName;
+		/** Setting for new join policy. */
+		TOptional<ELobbyJoinPolicy> JoinPolicy;
+		/** Setting for lobby ownership change. */
+		TOptional<FAccountId> OwnerAccountId;
+		/** Target member to be kicked. */
+		TOptional<FAccountId> KickedTargetMember;
+		/** Access to local member when joining and assigning initial attributes. */
+		TSharedPtr<FLobbyMemberInternal> JoiningLocalMember;
+		/** Access to local member when updating attributes. */
+		TSharedPtr<FLobbyMemberInternal> UpdatedLocalMember;
+		/** Reason to be set when the local member is leaving. */
+		TOptional<ELobbyMemberLeaveReason> LocalUserLeaveReason;
+		/** Whether there are waiting lobby attribute changes to be committed. */
+		bool bLobbyAttributesChanged = false;
+	};
+
+	/**
+	 * The schema registry is needed to initialize new members.
+	 */
+	TSharedRef<const FSchemaRegistry> SchemaRegistry;
 
 	/**
 	 * The shared pointer given back to user code with lobby operation results and notifications.
 	 * Any changes to this data are immediately available to users.
 	 */
-	TSharedRef<FLobby> PublicData;
+	TSharedRef<FLobbyInternal> InternalPublicData;
 
 	/** Mutable lobby member data storage. */
-	TMap<FAccountId, TSharedRef<FClientLobbyMemberSnapshot>> MemberDataStorage;
+	TMap<FAccountId, TSharedRef<FLobbyMemberInternal>> MemberDataStorage;
 
 	/**
 	 * Keep track of which members are local to the client.
 	 * When all local members have been removed all members will be removed.
 	 */
 	TSet<FAccountId> LocalMembers;
+
+	/** Updates which have been applied but not yet committed. */
+	mutable TOptional<FPreparedServiceChanges> PreparedServiceChanges;
+	mutable TOptional<FPreparedClientChanges> PreparedClientChanges;
 };
 
 #if LOBBIES_FUNCTIONAL_TEST_ENABLED
@@ -312,24 +375,6 @@ struct FFunctionalTestLobbies
 
 namespace Meta {
 
-BEGIN_ONLINE_STRUCT_META(FLobbyConfig)
-	ONLINE_STRUCT_FIELD(FLobbyConfig, RegisteredSchemas)
-END_ONLINE_STRUCT_META()
-
-
-BEGIN_ONLINE_STRUCT_META(FLobbySchemaAttributeConfig)
-	ONLINE_STRUCT_FIELD(FLobbySchemaAttributeConfig, Name),
-	ONLINE_STRUCT_FIELD(FLobbySchemaAttributeConfig, ServiceAttributeFieldId),
-	ONLINE_STRUCT_FIELD(FLobbySchemaAttributeConfig, MaxByteSize),
-	ONLINE_STRUCT_FIELD(FLobbySchemaAttributeConfig, Visibility)
-END_ONLINE_STRUCT_META()
-
-BEGIN_ONLINE_STRUCT_META(FLobbySchemaConfig)
-	ONLINE_STRUCT_FIELD(FLobbySchemaConfig, BaseSchemaName),
-	ONLINE_STRUCT_FIELD(FLobbySchemaConfig, LobbyAttributes),
-	ONLINE_STRUCT_FIELD(FLobbySchemaConfig, LobbyMemberAttributes)
-END_ONLINE_STRUCT_META()
-
 #if LOBBIES_FUNCTIONAL_TEST_ENABLED
 BEGIN_ONLINE_STRUCT_META(FFunctionalTestLobbies::Params)
 END_ONLINE_STRUCT_META()
