Integrate schema into lobbies.

diff --git a/Engine/Plugins/Online/Experimental/OnlineServices/Source/OnlineServicesInterface/Private/Online/Schema.cpp b/Engine/Plugins/Online/Experimental/OnlineServices/Source/OnlineServicesInterface/Private/Online/Schema.cpp
--- a/Engine/Plugins/Online/Experimental/OnlineServices/Source/OnlineServicesInterface/Private/Online/Schema.cpp
+++ b/Engine/Plugins/Online/Experimental/OnlineServices/Source/OnlineServicesInterface/Private/Online/Schema.cpp
@@ -762,7 +762,9 @@ bool FSchemaRegistry::ParseConfig(const FSchemaRegistryDescriptorConfig& Config)
 	}
 }
 
-FSchemaCategoryInstance::FSchemaCategoryInstance(
+namespace Private {
+
+FSchemaCategoryInstanceBase::FSchemaCategoryInstanceBase(
 	const FSchemaId& DerivedSchemaId,
 	const FSchemaId& BaseSchemaId,
 	const FSchemaCategoryId& CategoryId,
@@ -771,7 +773,7 @@ FSchemaCategoryInstance::FSchemaCategoryInstance(
 {
 	bool bInitSuccess = true;
 
-	if (InitializeSchemaDefinition(BaseSchemaId, CategoryId, &BaseSchemaDefinition, &BaseSchemaCategoryDefinition))
+	if (InitializeSchemaDefinition(BaseSchemaId, nullptr, CategoryId, &BaseSchemaDefinition, &BaseSchemaCategoryDefinition))
 	{
 		if (DerivedSchemaId == FSchemaId())
 		{
@@ -782,24 +784,7 @@ FSchemaCategoryInstance::FSchemaCategoryInstance(
 		}
 		else
 		{
-			if (InitializeSchemaDefinition(DerivedSchemaId, CategoryId, &DerivedSchemaDefinition, &DerivedSchemaCategoryDefinition))
-			{
-				// When a derived schema is in use the base schema must have a field flagged as SchemaCompatibilityId.
-				if (BaseSchemaCategoryDefinition->SchemaCompatibilityServiceAttributeId == FSchemaServiceAttributeId())
-				{
-					UE_LOG(LogOnlineSchema, Error, TEXT("[FSchemaCategoryInstance] init error: Base schema %s is missing a valid SchemaCompatibilityId field"),
-						*BaseSchemaId.ToString().ToLower());
-					bInitSuccess = false;
-				}
-				else if (!SchemaRegistry->IsSchemaChildOf(DerivedSchemaId, BaseSchemaId))
-				{
-					UE_LOG(LogOnlineSchema, Error, TEXT("[FSchemaCategoryInstance] init error: Schema %s is not a child of %s"),
-						*DerivedSchemaId.ToString().ToLower(),
-						*BaseSchemaId.ToString().ToLower());
-					bInitSuccess = false;
-				}
-			}
-			else
+			if (!InitializeSchemaDefinition(DerivedSchemaId, BaseSchemaDefinition.Get(), CategoryId, &DerivedSchemaDefinition, &DerivedSchemaCategoryDefinition))
 			{
 				// Unable to find derived schema when it was expected to be valid.
 				bInitSuccess = false;
@@ -809,6 +794,11 @@ FSchemaCategoryInstance::FSchemaCategoryInstance(
 
 	if (!bInitSuccess)
 	{
+		UE_LOG(LogOnlineSchema, Error, TEXT("[FSchemaCategoryInstance] init failed. DerivedSchema[%s], BaseSchema[%s], CategoryId[%s]"),
+			*DerivedSchemaId.ToString().ToLower(),
+			*BaseSchemaId.ToString().ToLower(),
+			*CategoryId.ToString().ToLower());
+
 		// Invalidate schema to fail initialization.
 		BaseSchemaDefinition.Reset();
 		BaseSchemaCategoryDefinition = nullptr;
@@ -817,109 +807,176 @@ FSchemaCategoryInstance::FSchemaCategoryInstance(
 	}
 }
 
-TOnlineResult<FSchemaCategoryInstanceApplyApplicationDelta> FSchemaCategoryInstance::ApplyApplicationDelta(FSchemaCategoryInstanceApplyApplicationDelta::Params&& Params)
+TOnlineResult<FSchemaCategoryInstancePrepareClientChanges> FSchemaCategoryInstanceBase::PrepareClientChanges(FSchemaCategoryInstancePrepareClientChanges::Params&& Params) const
 {
 	// todo: handle packing multiple attributes into a service attribute.
 
+	// Start by clearing any previous prepared state.
+	ResetPreparedChanges();
+
 	const FSchemaDefinition* SchemaDefinition = nullptr;
 	const FSchemaCategoryDefinition* SchemaCategoryDefinition = nullptr;
 	if (!GetSerializationSchema(&SchemaDefinition, &SchemaCategoryDefinition))
 	{
-		UE_LOG(LogOnlineSchema, Warning, TEXT("[ApplyApplicationDelta] Schema category definition not found."));
-		return TOnlineResult<FSchemaCategoryInstanceApplyApplicationDelta>(Errors::InvalidState());
+		UE_LOG(LogOnlineSchema, Warning, TEXT("[PrepareClientChanges] Serialization schema category definition not found."));
+		return TOnlineResult<FSchemaCategoryInstancePrepareClientChanges>(Errors::InvalidState());
 	}
 
-	// Configured service attributes must be present to serialize to service.
-	if (SchemaCategoryDefinition->ServiceAttributeDefinitions.IsEmpty())
+	FSchemaClientServiceChanges ServiceChanges;
+	FPreparedClientChanges NewPreparedClientChanges;
+
+	// Handle schema swap.
+	if (Params.ClientChanges.SchemaId)
 	{
-		UE_LOG(LogOnlineSchema, Warning, TEXT("[ApplyApplicationDelta] %s.%s: Schema category has not been configured for serialization."),
-			*SchemaDefinition->Id.ToString().ToLower(),
-			*SchemaCategoryDefinition->Id.ToString().ToLower());
-		return TOnlineResult<FSchemaCategoryInstanceApplyApplicationDelta>(Errors::InvalidParams());
-	}
+		if (InitializeSchemaDefinition(
+			*Params.ClientChanges.SchemaId,
+			BaseSchemaDefinition.Get(),
+			BaseSchemaCategoryDefinition->Id,
+			&NewPreparedClientChanges.DerivedSchemaDefinition,
+			&NewPreparedClientChanges.DerivedSchemaCategoryDefinition))
+		{
+			if (DerivedSchemaDefinition)
+			{
+				// todo: handle schema change when previously set.
+				return TOnlineResult<FSchemaCategoryInstancePrepareClientChanges>(Errors::NotImplemented());
+			}
+			else
+			{
+				// Derived schema init succeeded. Switch to using it for preparing the service changes.
+				SchemaDefinition = NewPreparedClientChanges.DerivedSchemaDefinition.Get();
+				SchemaCategoryDefinition = NewPreparedClientChanges.DerivedSchemaCategoryDefinition;
 
-	FSchemaServiceChanges Changes;
+				NewPreparedClientChanges.ClientChanges.SchemaId.Emplace(SchemaDefinition->Id);
+			}
+		}
+		else
+		{
+			UE_LOG(LogOnlineSchema, Warning, TEXT("[PrepareClientChanges] Schema change failed: NewDerivedSchema[%s], BaseSchema[%s]"),
+				*Params.ClientChanges.SchemaId->ToString().ToLower(),
+				*BaseSchemaCategoryDefinition->Id.ToString().ToLower());
+			return TOnlineResult<FSchemaCategoryInstancePrepareClientChanges>(Errors::InvalidParams());
+		}
+	}
 
 	// Translate mutated attributes.
-	for (TPair<FSchemaAttributeId, FSchemaVariant>& UpdatedAttributeData : Params.UpdatedAttributes)
+	for (TPair<FSchemaAttributeId, FSchemaVariant>& UpdatedAttributeData : Params.ClientChanges.UpdatedAttributes)
 	{
-		const FSchemaAttributeId& UpdatedAttributeId = UpdatedAttributeData.Get<0>();
-		FSchemaVariant& UpdatedAttributeValue = UpdatedAttributeData.Get<1>();
+		const FSchemaAttributeId& UpdatedAttributeId = UpdatedAttributeData.Key;
+		FSchemaVariant& UpdatedAttributeValue = UpdatedAttributeData.Value;
 
-		if (const FSchemaAttributeDefinition* AttributeDefinition = SchemaCategoryDefinition->SchemaAttributeDefinitions.Find(UpdatedAttributeId))
+		const FSchemaAttributeDefinition* AttributeDefinition = SchemaCategoryDefinition->SchemaAttributeDefinitions.Find(UpdatedAttributeId);
+		if (AttributeDefinition == nullptr)
 		{
-			// Schema compatibility id cannot be modified by application code.
-			if (EnumHasAnyFlags(AttributeDefinition->Flags, ESchemaAttributeFlags::SchemaCompatibilityId))
+			UE_LOG(LogOnlineSchema, Warning, TEXT("[PrepareClientChanges] %s.%s.%s: Attribute definition does not exist in schema category."),
+				*SchemaDefinition->Id.ToString().ToLower(),
+				*SchemaCategoryDefinition->Id.ToString().ToLower(),
+				*UpdatedAttributeId.ToString().ToLower());
+			return TOnlineResult<FSchemaCategoryInstancePrepareClientChanges>(Errors::InvalidParams());
+		}
+
+		// Schema compatibility id cannot be modified by client code.
+		if (EnumHasAnyFlags(AttributeDefinition->Flags, ESchemaAttributeFlags::SchemaCompatibilityId))
+		{
+			UE_LOG(LogOnlineSchema, Warning, TEXT("[PrepareClientChanges] %s.%s.%s: It is not valid to set SchemaCompatibilityId in an attribute update"),
+				*SchemaDefinition->Id.ToString().ToLower(),
+				*SchemaCategoryDefinition->Id.ToString().ToLower(),
+				*AttributeDefinition->Id.ToString().ToLower());
+			return TOnlineResult<FSchemaCategoryInstancePrepareClientChanges>(Errors::InvalidParams());
+		}
+
+		// Check that type matches.
+		if (AttributeDefinition->Type != UpdatedAttributeValue.GetType())
+		{
+			UE_LOG(LogOnlineSchema, Warning, TEXT("[PrepareClientChanges] %s.%s.%s: Attribute type %s does not match the expected type %s."),
+				*SchemaDefinition->Id.ToString().ToLower(),
+				*SchemaCategoryDefinition->Id.ToString().ToLower(),
+				*AttributeDefinition->Id.ToString().ToLower(),
+				LexToString(UpdatedAttributeValue.GetType()),
+				LexToString(AttributeDefinition->Type));
+			return TOnlineResult<FSchemaCategoryInstancePrepareClientChanges>(Errors::InvalidParams());
+		}
+
+		// Check that variable sized data fits within the service attribute.
+		if (AttributeDefinition->Type == ESchemaAttributeType::String && UpdatedAttributeValue.GetString().Len() > AttributeDefinition->MaxSize)
+		{
+			UE_LOG(LogOnlineSchema, Warning, TEXT("[PrepareClientChanges] %s.%s.%s: Variably sized attribute exceeds maximum defined length. %d > %d."),
+				*SchemaDefinition->Id.ToString().ToLower(),
+				*SchemaCategoryDefinition->Id.ToString().ToLower(),
+				*AttributeDefinition->Id.ToString().ToLower(),
+				UpdatedAttributeValue.GetString().Len(),
+				AttributeDefinition->MaxSize);
+			return TOnlineResult<FSchemaCategoryInstancePrepareClientChanges>(Errors::InvalidParams());
+		}
+
+		// Compare attribute to previously known client data.
+		if (const FSchemaVariant* PreviousValue = GetClientSnapshot().Find(UpdatedAttributeId))
+		{
+			// Ignore change to the same value.
+			if (*PreviousValue != UpdatedAttributeValue)
 			{
-				UE_LOG(LogOnlineSchema, Warning, TEXT("[ApplyApplicationDelta] %s.%s.%s: It is not valid to set SchemaCompatibilityId in an attribute update"),
-					*SchemaDefinition->Id.ToString().ToLower(),
-					*SchemaCategoryDefinition->Id.ToString().ToLower(),
-					*AttributeDefinition->Id.ToString().ToLower());
-				return TOnlineResult<FSchemaCategoryInstanceApplyApplicationDelta>(Errors::InvalidParams());
+				// Store change to client data.
+				NewPreparedClientChanges.ClientChanges.ChangedAttributes.Add(UpdatedAttributeId, { *PreviousValue, UpdatedAttributeValue });
 			}
+		}
+		else
+		{
+			// Store change to client data.
+			NewPreparedClientChanges.ClientChanges.AddedAttributes.Add(UpdatedAttributeId, UpdatedAttributeValue);
+		}
 
-			// Check that type matches.
-			if (AttributeDefinition->Type != UpdatedAttributeValue.GetType())
+		if (AttributeDefinition->ServiceAttributeId != FSchemaServiceAttributeId())
+		{
+			const FSchemaServiceAttributeDefinition* ServiceAttributeDefinition = SchemaCategoryDefinition->ServiceAttributeDefinitions.Find(AttributeDefinition->ServiceAttributeId);
+			if (ensure(ServiceAttributeDefinition))
 			{
-				UE_LOG(LogOnlineSchema, Warning, TEXT("[ApplyApplicationDelta] %s.%s.%s: Attribute type %s does not match the expected type %s."),
-					*SchemaDefinition->Id.ToString().ToLower(),
-					*SchemaCategoryDefinition->Id.ToString().ToLower(),
-					*AttributeDefinition->Id.ToString().ToLower(),
-					LexToString(UpdatedAttributeValue.GetType()),
-					LexToString(AttributeDefinition->Type));
-				return TOnlineResult<FSchemaCategoryInstanceApplyApplicationDelta>(Errors::InvalidParams());
+				// Add modified attribute to output.
+				ServiceChanges.UpdatedAttributes.Add(AttributeDefinition->ServiceAttributeId, { AttributeDefinition->ServiceAttributeId, ServiceAttributeDefinition->Flags, MoveTemp(UpdatedAttributeValue) });
 			}
-
-			// Check that variable sized data fits within the service attribute.
-			if (AttributeDefinition->Type == ESchemaAttributeType::String && UpdatedAttributeValue.GetString().Len() > AttributeDefinition->MaxSize)
+			else
 			{
-				UE_LOG(LogOnlineSchema, Warning, TEXT("[ApplyApplicationDelta] %s.%s.%s: Variably sized attribute exceeds maximum defined length. %d > %d."),
+				UE_LOG(LogOnlineSchema, Warning, TEXT("[PrepareClientChanges] %s.%s.%s: Unable to find service attribute definition for service attribute %s."),
 					*SchemaDefinition->Id.ToString().ToLower(),
 					*SchemaCategoryDefinition->Id.ToString().ToLower(),
 					*AttributeDefinition->Id.ToString().ToLower(),
-					UpdatedAttributeValue.GetString().Len(),
-					AttributeDefinition->MaxSize);
-				return TOnlineResult<FSchemaCategoryInstanceApplyApplicationDelta>(Errors::InvalidParams());
+					*AttributeDefinition->ServiceAttributeId.ToString().ToLower());
+				return TOnlineResult<FSchemaCategoryInstancePrepareClientChanges>(Errors::InvalidState());
 			}
-
-			Changes.UpdatedAttributes.Add(AttributeDefinition->ServiceAttributeId, {AttributeDefinition->ServiceAttributeId, AttributeDefinition->Flags, MoveTemp(UpdatedAttributeValue)});
-		}
-		else
-		{
-			UE_LOG(LogOnlineSchema, Warning, TEXT("[ApplyApplicationDelta] %s.%s.%s: Attribute definition does not exist in schema category."),
-				*SchemaDefinition->Id.ToString().ToLower(),
-				*SchemaCategoryDefinition->Id.ToString().ToLower(),
-				*UpdatedAttributeId.ToString().ToLower());
-			return TOnlineResult<FSchemaCategoryInstanceApplyApplicationDelta>(Errors::InvalidParams());
 		}
 	}
 
 	// Translate removed attributes.
-	for (FSchemaAttributeId& RemovedAttributeId : Params.RemovedAttributes)
+	for (FSchemaAttributeId& RemovedAttributeId : Params.ClientChanges.RemovedAttributes)
 	{
 		if (const FSchemaAttributeDefinition* AttributeDefinition = SchemaCategoryDefinition->SchemaAttributeDefinitions.Find(RemovedAttributeId))
 		{
-			// Schema compatibility id cannot be modified by application code.
+			// Schema compatibility id cannot be modified by client code.
 			if (EnumHasAnyFlags(AttributeDefinition->Flags, ESchemaAttributeFlags::SchemaCompatibilityId))
 			{
-				UE_LOG(LogOnlineSchema, Warning, TEXT("[ApplyApplicationDelta] %s.%s.%s: It is not valid to set SchemaCompatibilityId in an attribute update"),
+				UE_LOG(LogOnlineSchema, Warning, TEXT("[PrepareClientChanges] %s.%s.%s: It is not valid to set SchemaCompatibilityId in an attribute update"),
 					*SchemaDefinition->Id.ToString().ToLower(),
 					*SchemaCategoryDefinition->Id.ToString().ToLower(),
 					*AttributeDefinition->Id.ToString().ToLower());
-				return TOnlineResult<FSchemaCategoryInstanceApplyApplicationDelta>(Errors::InvalidParams());
+				return TOnlineResult<FSchemaCategoryInstancePrepareClientChanges>(Errors::InvalidParams());
 			}
 			else
 			{
-				Changes.RemovedAttributes.Add(AttributeDefinition->ServiceAttributeId);
+				// Store change to client data.
+				NewPreparedClientChanges.ClientChanges.RemovedAttributes.Add(RemovedAttributeId);
+
+				if (AttributeDefinition->ServiceAttributeId != FSchemaServiceAttributeId())
+				{
+					// Add removed attribute to service output.
+					ServiceChanges.RemovedAttributes.Add(AttributeDefinition->ServiceAttributeId);
+				}
 			}
 		}
 		else
 		{
-			UE_LOG(LogOnlineSchema, Warning, TEXT("[ApplyApplicationDelta] %s.%s.%s: Attribute definition does not exist in schema category."),
+			UE_LOG(LogOnlineSchema, Warning, TEXT("[PrepareClientChanges] %s.%s.%s: Attribute definition does not exist in schema category."),
 				*SchemaDefinition->Id.ToString().ToLower(),
 				*SchemaCategoryDefinition->Id.ToString().ToLower(),
 				*RemovedAttributeId.ToString().ToLower());
-			return TOnlineResult<FSchemaCategoryInstanceApplyApplicationDelta>(Errors::InvalidParams());
+			return TOnlineResult<FSchemaCategoryInstancePrepareClientChanges>(Errors::InvalidParams());
 		}
 	}
 
@@ -929,114 +986,202 @@ TOnlineResult<FSchemaCategoryInstanceApplyApplicationDelta> FSchemaCategoryInsta
 		const FSchemaAttributeDefinition* AttributeDefinition = SchemaCategoryDefinition->SchemaAttributeDefinitions.Find(SchemaCategoryDefinition->SchemaCompatibilityAttributeId);
 		if (ensure(AttributeDefinition))
 		{
-			Changes.UpdatedAttributes.Add(AttributeDefinition->ServiceAttributeId, { AttributeDefinition->ServiceAttributeId, AttributeDefinition->Flags, SchemaDefinition->CompatibilityId });
-			LastSentSchemaCompatibilityId = SchemaDefinition->CompatibilityId;
+			if (AttributeDefinition->ServiceAttributeId != FSchemaServiceAttributeId())
+			{
+				const FSchemaServiceAttributeDefinition* ServiceAttributeDefinition = SchemaCategoryDefinition->ServiceAttributeDefinitions.Find(AttributeDefinition->ServiceAttributeId);
+				if (ensure(ServiceAttributeDefinition))
+				{
+					ServiceChanges.UpdatedAttributes.Add(AttributeDefinition->ServiceAttributeId, { AttributeDefinition->ServiceAttributeId, ServiceAttributeDefinition->Flags, SchemaDefinition->CompatibilityId });
+				}
+				else
+				{
+					UE_LOG(LogOnlineSchema, Warning, TEXT("[PrepareClientChanges] %s.%s.%s: Unable to find service attribute definition for service attribute %s."),
+						*SchemaDefinition->Id.ToString().ToLower(),
+						*SchemaCategoryDefinition->Id.ToString().ToLower(),
+						*AttributeDefinition->Id.ToString().ToLower(),
+						*AttributeDefinition->ServiceAttributeId.ToString().ToLower());
+					return TOnlineResult<FSchemaCategoryInstancePrepareClientChanges>(Errors::InvalidState());
+				}
+			}
+
+			NewPreparedClientChanges.SchemaCompatibilityId = SchemaDefinition->CompatibilityId;
 		}
 		else
 		{
-			UE_LOG(LogOnlineSchema, Warning, TEXT("[ApplyApplicationDelta] %s.%s.%s: Failed to find attribute defnition for SchemaCompatibilityId field."),
+			UE_LOG(LogOnlineSchema, Warning, TEXT("[PrepareClientChanges] %s.%s.%s: Failed to find attribute definition for SchemaCompatibilityId field."),
 				*SchemaDefinition->Id.ToString().ToLower(),
 				*SchemaCategoryDefinition->Id.ToString().ToLower(),
 				*SchemaCategoryDefinition->SchemaCompatibilityAttributeId.ToString().ToLower());
-			return TOnlineResult<FSchemaCategoryInstanceApplyApplicationDelta>(Errors::InvalidState());
+			return TOnlineResult<FSchemaCategoryInstancePrepareClientChanges>(Errors::InvalidState());
 		}
 	}
 
-	return TOnlineResult<FSchemaCategoryInstanceApplyApplicationDelta>(FSchemaCategoryInstanceApplyApplicationDelta::Result({MoveTemp(Changes)}));
+	// Store pending changes to be applied on commit.
+	PreparedClientChanges = MoveTemp(NewPreparedClientChanges);
+
+	return TOnlineResult<FSchemaCategoryInstancePrepareClientChanges>({MoveTemp(ServiceChanges)});
+}
+
+FSchemaCategoryInstanceCommitClientChanges::Result FSchemaCategoryInstanceBase::CommitClientChanges()
+{
+	FSchemaServiceClientChanges ClientChanges;
+
+	// Check that there is a waiting delta state to be applied.
+	if (!PreparedClientChanges.IsSet())
+	{
+		return FSchemaCategoryInstanceCommitClientChanges::Result({ MoveTemp(ClientChanges) });
+	}
+
+	// Apply delta to state snapshot.
+	TMap<FSchemaAttributeId, FSchemaVariant>& ClientDataSnapshot = GetMutableClientSnapshot();
+
+	// Handle added attributes.
+	Algo::ForEach(PreparedClientChanges->ClientChanges.AddedAttributes,
+		[&ClientDataSnapshot](const TPair<FSchemaAttributeId, FSchemaVariant>& AddedPair) -> void
+		{
+			ClientDataSnapshot.Add(AddedPair.Key, AddedPair.Value);
+		});
+
+	// Handle changed attributes.
+	Algo::ForEach(PreparedClientChanges->ClientChanges.ChangedAttributes,
+		[&ClientDataSnapshot](const TPair<FSchemaAttributeId, TPair<FSchemaVariant, FSchemaVariant>>& ChangedPair) -> void
+		{
+			ClientDataSnapshot.Add(ChangedPair.Key, ChangedPair.Value.Value);
+		});
+
+	// Handle removed attributes.
+	Algo::ForEach(PreparedClientChanges->ClientChanges.RemovedAttributes,
+		[&ClientDataSnapshot](FSchemaAttributeId RemovedAttributeId) -> void
+		{
+			ClientDataSnapshot.Remove(RemovedAttributeId);
+		});
+
+	// Set last sent schema compatibility id.
+	if (PreparedClientChanges->SchemaCompatibilityId != 0)
+	{
+		LastSentSchemaCompatibilityId = PreparedClientChanges->SchemaCompatibilityId;
+	}
+
+	// Set schema definitions from a schema swap.
+	if (PreparedClientChanges->DerivedSchemaCategoryDefinition)
+	{
+		DerivedSchemaDefinition = PreparedClientChanges->DerivedSchemaDefinition;
+		DerivedSchemaCategoryDefinition = PreparedClientChanges->DerivedSchemaCategoryDefinition;
+	}
+
+	// Move out client changes for returning to the client.
+	ClientChanges = MoveTemp(PreparedClientChanges->ClientChanges);
+
+	// Reset pending state.
+	ResetPreparedChanges();
+
+	// Return client changes.
+	return FSchemaCategoryInstanceCommitClientChanges::Result({MoveTemp(ClientChanges)});
 }
 
-TOnlineResult<FSchemaCategoryInstanceApplyServiceSnapshot> FSchemaCategoryInstance::ApplyServiceSnapshot(FSchemaCategoryInstanceApplyServiceSnapshot::Params&& Params)
+TOnlineResult<FSchemaCategoryInstancePrepareServiceSnapshot> FSchemaCategoryInstanceBase::PrepareServiceSnapshot(FSchemaCategoryInstancePrepareServiceSnapshot::Params&& Params) const
 {
 	// todo: handle unpacking a service attribute into multiple attributes.
 
+	// Start by clearing any previous prepared state.
+	ResetPreparedChanges();
+	FPreparedServiceChanges NewPreparedServiceChanges;
+	FSchemaCategoryInstancePrepareServiceSnapshot::Result ApplyResult;
+
 	const FSchemaDefinition* SchemaDefinition = nullptr;
 	const FSchemaCategoryDefinition* SchemaCategoryDefinition = nullptr;
-
 	if (!GetSerializationSchema(&SchemaDefinition, &SchemaCategoryDefinition))
 	{
-		UE_LOG(LogOnlineSchema, Warning, TEXT("[ApplyServiceSnapshot] Schema category definition not found."));
-		return TOnlineResult<FSchemaCategoryInstanceApplyServiceSnapshot>(Errors::InvalidState());
+		UE_LOG(LogOnlineSchema, Warning, TEXT("[CommitServiceSnapshot] Serialization schema category definition not found."));
+		return TOnlineResult<FSchemaCategoryInstancePrepareServiceSnapshot>(Errors::InvalidState());
 	}
 
-	FSchemaApplicationChanges Changes;
-
 	// Check for a schema change.
-	TPair<FSchemaServiceAttributeId, FSchemaVariant>* SchemaCompatibilityAttributeData = Algo::FindByPredicate(Params.Attributes,
+	TPair<FSchemaServiceAttributeId, FSchemaVariant>* SchemaCompatibilityAttributeData = Algo::FindByPredicate(Params.ServiceSnapshot.Attributes,
 		[&SchemaCategoryDefinition](const TPair<FSchemaServiceAttributeId, FSchemaVariant>& AttributeData)
 		{
-			return AttributeData.Get<0>() == SchemaCategoryDefinition->SchemaCompatibilityServiceAttributeId;
+			return AttributeData.Key == SchemaCategoryDefinition->SchemaCompatibilityServiceAttributeId;
 		});
 
-	// Handle incoming schema change.
+	// Handle incoming schema change from service attribute.
 	if (SchemaCompatibilityAttributeData)
 	{
-		const FSchemaVariant& SchemaCompatibilityIdValue = SchemaCompatibilityAttributeData->Get<1>();
-
-		if (TSharedPtr<const FSchemaDefinition> SwappedSchema = SchemaRegistry->GetDefinition(SchemaCompatibilityIdValue.GetInt64()))
+		const int64 SchemaCompatibilityId = SchemaCompatibilityAttributeData->Value.GetInt64();
+
+		if (InitializeSchemaDefinition(
+			SchemaCompatibilityId,
+			BaseSchemaDefinition.Get(),
+			SchemaCategoryDefinition->Id,
+			&NewPreparedServiceChanges.DerivedSchemaDefinition,
+			&NewPreparedServiceChanges.DerivedSchemaCategoryDefinition))
 		{
-			if (!SchemaRegistry->IsSchemaChildOf(SwappedSchema->Id, BaseSchemaDefinition->Id))
+			// Check that schema has actually changed.
+			if (DerivedSchemaDefinition == nullptr || DerivedSchemaDefinition->Id != NewPreparedServiceChanges.DerivedSchemaDefinition->Id)
 			{
-				UE_LOG(LogOnlineSchema, Warning, TEXT("[ApplyServiceSnapshot] %s.%s.%s: Invalid schema swap. Incoming schema %s is not a child of %s."),
-					*SchemaDefinition->Id.ToString().ToLower(),
-					*SchemaCategoryDefinition->Id.ToString().ToLower(),
-					*SchemaCategoryDefinition->SchemaCompatibilityServiceAttributeId.ToString().ToLower(),
-					*SwappedSchema->Id.ToString().ToLower(),
-					*BaseSchemaDefinition->Id.ToString().ToLower());
-				return TOnlineResult<FSchemaCategoryInstanceApplyServiceSnapshot>(Errors::InvalidParams());
+				// Derived schema init succeeded. Switch to using it for preparing the client changes.
+				SchemaDefinition = NewPreparedServiceChanges.DerivedSchemaDefinition.Get();
+				SchemaCategoryDefinition = NewPreparedServiceChanges.DerivedSchemaCategoryDefinition;
+
+				NewPreparedServiceChanges.ClientChanges.SchemaId.Emplace(SchemaDefinition->Id);
+				ApplyResult.DerivedSchemaId.Emplace(SchemaDefinition->Id);
 			}
+		}
+		else
+		{
+			UE_LOG(LogOnlineSchema, Warning, TEXT("[PrepareServiceSnapshot] %s.%s.%s: Failed to find valid definition for SchemaCompatibilityId 0x%08" INT64_X_FMT "."),
+				*SchemaDefinition->Id.ToString().ToLower(),
+				*SchemaCategoryDefinition->Id.ToString().ToLower(),
+				*SchemaCategoryDefinition->SchemaCompatibilityServiceAttributeId.ToString().ToLower(),
+				SchemaCompatibilityId);
+			return TOnlineResult<FSchemaCategoryInstancePrepareServiceSnapshot>(Errors::InvalidParams());
+		}
+	}
 
+	// Handle incoming schema change from user parameter.
+	if (Params.SchemaId)
+	{
+		// Manual setting of schema id is only allowed when it cannot be detected from the service.
+		if (SchemaCompatibilityAttributeData)
+		{
+			UE_LOG(LogOnlineSchema, Warning, TEXT("[PrepareServiceSnapshot] %s.%s: Unable to handle user schema swap when service defines a schema compatibility id."),
+				*SchemaDefinition->Id.ToString().ToLower(),
+				*SchemaCategoryDefinition->Id.ToString().ToLower());
+			return TOnlineResult<FSchemaCategoryInstancePrepareServiceSnapshot>(Errors::InvalidParams());
+		}
+	
+		if (InitializeSchemaDefinition(
+			*Params.SchemaId,
+			BaseSchemaDefinition.Get(),
+			SchemaCategoryDefinition->Id,
+			&NewPreparedServiceChanges.DerivedSchemaDefinition,
+			&NewPreparedServiceChanges.DerivedSchemaCategoryDefinition))
+		{
 			// Check that schema has actually changed.
-			if (DerivedSchemaDefinition == nullptr || DerivedSchemaDefinition->Id != SwappedSchema->Id)
+			if (DerivedSchemaDefinition == nullptr || DerivedSchemaDefinition->Id != NewPreparedServiceChanges.DerivedSchemaDefinition->Id)
 			{
-				Changes.SchemaId.Emplace(SwappedSchema->Id);
+				// Derived schema init succeeded. Switch to using it for preparing the client changes.
+				SchemaDefinition = NewPreparedServiceChanges.DerivedSchemaDefinition.Get();
+				SchemaCategoryDefinition = NewPreparedServiceChanges.DerivedSchemaCategoryDefinition;
 
-				// Category must exist due to parent -> child relationship
-				const FSchemaCategoryDefinition* SwappedSchemaCategoryDefinition = SwappedSchema->Categories.Find(SchemaCategoryDefinition->Id);
-				if (ensure(SwappedSchemaCategoryDefinition))
-				{
-					UE_LOG(LogOnlineSchema, Verbose, TEXT("[ApplyServiceSnapshot] %s.%s.%s: Derived schema changed from %s to %s."),
-						*SchemaDefinition->Id.ToString().ToLower(),
-						*SchemaCategoryDefinition->Id.ToString().ToLower(),
-						*SchemaCategoryDefinition->SchemaCompatibilityServiceAttributeId.ToString().ToLower(),
-						DerivedSchemaDefinition ? *DerivedSchemaDefinition->Id.ToString().ToLower() : TEXT("nullptr"),
-						*SwappedSchemaCategoryDefinition->Id.ToString().ToLower());
-
-					// Set new schema to class member data.
-					DerivedSchemaDefinition = SwappedSchema;
-					DerivedSchemaCategoryDefinition = SwappedSchemaCategoryDefinition;
-
-					// Update schema definitions for current update processing.
-					SchemaDefinition = DerivedSchemaDefinition.Get();
-					SchemaCategoryDefinition = SwappedSchemaCategoryDefinition;
-				}
-				else
-				{
-					UE_LOG(LogOnlineSchema, Error, TEXT("[ApplyServiceSnapshot] %s.%s.%s: Failed to find expected category definition in incoming schema %s."),
-						*SchemaDefinition->Id.ToString().ToLower(),
-						*SchemaCategoryDefinition->Id.ToString().ToLower(),
-						*SchemaCategoryDefinition->SchemaCompatibilityServiceAttributeId.ToString().ToLower(),
-						*SchemaCategoryDefinition->Id.ToString().ToLower());
-					return TOnlineResult<FSchemaCategoryInstanceApplyServiceSnapshot>(Errors::InvalidParams());
-				}
+				NewPreparedServiceChanges.ClientChanges.SchemaId.Emplace(SchemaDefinition->Id);
+				ApplyResult.DerivedSchemaId.Emplace(SchemaDefinition->Id);
 			}
 		}
 		else
 		{
-			UE_LOG(LogOnlineSchema, Warning, TEXT("[ApplyServiceSnapshot] %s.%s.%s: Failed to find definition for SchemaCompatibilityId 0x%08" INT64_X_FMT "."),
+			UE_LOG(LogOnlineSchema, Warning, TEXT("[PrepareServiceSnapshot] %s.%s: Failed to find valid definition for SchemaId %s."),
 				*SchemaDefinition->Id.ToString().ToLower(),
 				*SchemaCategoryDefinition->Id.ToString().ToLower(),
-				*SchemaCategoryDefinition->SchemaCompatibilityServiceAttributeId.ToString().ToLower(),
-				SchemaCompatibilityIdValue.GetInt64());
-			return TOnlineResult<FSchemaCategoryInstanceApplyServiceSnapshot>(Errors::InvalidParams());
+				*Params.SchemaId->ToString().ToLower());
+			return TOnlineResult<FSchemaCategoryInstancePrepareServiceSnapshot>(Errors::InvalidParams());
 		}
 	}
 
 	// Unpack service attributes.
-	TMap<FSchemaAttributeId, FSchemaVariant> NewApplicationDataServiceSnapshot;
-	for (TPair<FSchemaServiceAttributeId, FSchemaVariant>& SchemaServiceAttributeData : Params.Attributes)
+	for (TPair<FSchemaServiceAttributeId, FSchemaVariant>& SchemaServiceAttributeData : Params.ServiceSnapshot.Attributes)
 	{
-		const FSchemaServiceAttributeId& SchemaServiceAttributeId = SchemaServiceAttributeData.Get<0>();
-		FSchemaVariant& SchemaServiceAttributeValue = SchemaServiceAttributeData.Get<1>();
+		const FSchemaServiceAttributeId& SchemaServiceAttributeId = SchemaServiceAttributeData.Key;
+		FSchemaVariant& SchemaServiceAttributeValue = SchemaServiceAttributeData.Value;
 
 		if (const FSchemaServiceAttributeDefinition* SchemaServiceAttributeDefinition =
 			SchemaCategoryDefinition->ServiceAttributeDefinitions.Find(SchemaServiceAttributeId))
@@ -1051,10 +1196,10 @@ TOnlineResult<FSchemaCategoryInstanceApplyServiceSnapshot> FSchemaCategoryInstan
 				// Attribute definition must be valid due to successfully processing config.
 				check(SchemaAttributeDefinition);
 
-				// Schema changes are handled by surfacing a new schema id in the result object. Don't return SchemaCompatibilityId as an attribute to the application.
+				// Schema changes are handled by surfacing a new schema id in the result object. Don't return SchemaCompatibilityId as an attribute to the client.
 				if (EnumHasAnyFlags(SchemaAttributeDefinition->Flags, ESchemaAttributeFlags::SchemaCompatibilityId))
 				{
-					UE_LOG(LogOnlineSchema, Verbose, TEXT("[ApplyServiceSnapshot] %s.%s.%s: Consuming SchemaCompatibility attribute data: %s."),
+					UE_LOG(LogOnlineSchema, Verbose, TEXT("[PrepareServiceSnapshot] %s.%s.%s: Consuming SchemaCompatibility attribute data: %s."),
 						*SchemaDefinition->Id.ToString().ToLower(),
 						*SchemaCategoryDefinition->Id.ToString().ToLower(),
 						*SchemaAttributeDefinition->Id.ToString().ToLower(),
@@ -1062,62 +1207,100 @@ TOnlineResult<FSchemaCategoryInstanceApplyServiceSnapshot> FSchemaCategoryInstan
 				}
 				else
 				{
-					NewApplicationDataServiceSnapshot.Emplace(SchemaAttributeDefinition->Id, MoveTemp(SchemaServiceAttributeValue));
+					NewPreparedServiceChanges.ClientDataSnapshot.Emplace(SchemaAttributeDefinition->Id, MoveTemp(SchemaServiceAttributeValue));
 				}
 			}
 		}
 		else
 		{
-			UE_LOG(LogOnlineSchema, Warning, TEXT("[ApplyServiceSnapshot] %s.%s.%s: Service attribute definition does not exist in schema category."),
+			UE_LOG(LogOnlineSchema, Warning, TEXT("[PrepareServiceSnapshot] %s.%s.%s: Service attribute definition does not exist in schema category."),
 				*SchemaDefinition->Id.ToString().ToLower(),
 				*SchemaCategoryDefinition->Id.ToString().ToLower(),
 				*SchemaServiceAttributeId.ToString().ToLower());
-			return TOnlineResult<FSchemaCategoryInstanceApplyServiceSnapshot>(Errors::InvalidParams());
+			return TOnlineResult<FSchemaCategoryInstancePrepareServiceSnapshot>(Errors::InvalidParams());
 		}
 	}
 
-	// Process new snapshot against old snapshot to build application data changes.
+	// Process new snapshot against old snapshot to build client data changes.
 	{
 		// Find attributes that were added or changed.
-		for (const TPair<FSchemaAttributeId, FSchemaVariant>& NewSnapshotAttributeData : NewApplicationDataServiceSnapshot)
+		for (const TPair<FSchemaAttributeId, FSchemaVariant>& NewSnapshotAttributeData : NewPreparedServiceChanges.ClientDataSnapshot)
 		{
-			const FSchemaServiceAttributeId& NewSnapshotAttributeId = NewSnapshotAttributeData.Get<0>();
-			const FSchemaVariant& NewSnapshotAttributeValue = NewSnapshotAttributeData.Get<1>();
+			const FSchemaServiceAttributeId& NewSnapshotAttributeId = NewSnapshotAttributeData.Key;
+			const FSchemaVariant& NewSnapshotAttributeValue = NewSnapshotAttributeData.Value;
 
-			if (FSchemaVariant* OldSnapshotAttributeValue = ApplicationDataServiceSnapshot.Find(NewSnapshotAttributeId))
+			if (const FSchemaVariant* OldSnapshotAttributeValue = GetClientSnapshot().Find(NewSnapshotAttributeId))
 			{
-				Changes.ChangedAttributes.Add(NewSnapshotAttributeId, { MoveTemp(*OldSnapshotAttributeValue), NewSnapshotAttributeValue });
+				if (*OldSnapshotAttributeValue != NewSnapshotAttributeValue)
+				{
+					NewPreparedServiceChanges.ClientChanges.ChangedAttributes.Add(NewSnapshotAttributeId, { *OldSnapshotAttributeValue, NewSnapshotAttributeValue });
+				}
 			}
 			else
 			{
-				Changes.AddedAttributes.Add(NewSnapshotAttributeId, NewSnapshotAttributeValue);
+				NewPreparedServiceChanges.ClientChanges.AddedAttributes.Add(NewSnapshotAttributeId, NewSnapshotAttributeValue);
 			}
 		}
 
 		// Add removed attribute ids to changes.
-		Algo::TransformIf(ApplicationDataServiceSnapshot, Changes.RemovedAttributes,
-			[&NewApplicationDataServiceSnapshot](const TPair<FSchemaAttributeId, FSchemaVariant>& OldSnapshotAttributeData) -> bool
+		Algo::TransformIf(GetClientSnapshot(), NewPreparedServiceChanges.ClientChanges.RemovedAttributes,
+			[&NewPreparedServiceChanges](const TPair<FSchemaAttributeId, FSchemaVariant>& OldSnapshotAttributeData) -> bool
 			{
-				const FSchemaServiceAttributeId& OldSnapshotAttributeId = OldSnapshotAttributeData.Get<0>();
-				return !NewApplicationDataServiceSnapshot.Contains(OldSnapshotAttributeId);
+				const FSchemaServiceAttributeId& OldSnapshotAttributeId = OldSnapshotAttributeData.Key;
+				return !NewPreparedServiceChanges.ClientDataSnapshot.Contains(OldSnapshotAttributeId);
 			},
 			[](const TPair<FSchemaAttributeId, FSchemaVariant>& OldSnapshotAttributeData) -> FSchemaAttributeId
 			{
-				const FSchemaServiceAttributeId& OldSnapshotAttributeId = OldSnapshotAttributeData.Get<0>();
+				const FSchemaServiceAttributeId& OldSnapshotAttributeId = OldSnapshotAttributeData.Key;
 				return OldSnapshotAttributeId;
 			});
 	}
 
-	ApplicationDataServiceSnapshot = MoveTemp(NewApplicationDataServiceSnapshot);
-	return TOnlineResult<FSchemaCategoryInstanceApplyServiceSnapshot>(FSchemaCategoryInstanceApplyServiceSnapshot::Result({MoveTemp(Changes)}));
+	// Store pending changes to be applied on commit.
+	PreparedServiceChanges = MoveTemp(NewPreparedServiceChanges);
+	return TOnlineResult<FSchemaCategoryInstancePrepareServiceSnapshot>(MoveTemp(ApplyResult));
+}
+
+FSchemaCategoryInstanceCommitServiceSnapshot::Result FSchemaCategoryInstanceBase::CommitServiceSnapshot()
+{
+	FSchemaServiceClientChanges ClientChanges;
+
+	// Check that there is a waiting snapshot state to be applied.
+	if (!PreparedServiceChanges.IsSet())
+	{
+		return FSchemaCategoryInstanceCommitServiceSnapshot::Result({ MoveTemp(ClientChanges) });
+	}
+
+	// Apply attribute changes.
+	GetMutableClientSnapshot() = MoveTemp(PreparedServiceChanges->ClientDataSnapshot);
+
+	// Set schema definitions from a schema swap.
+	if (PreparedServiceChanges->DerivedSchemaCategoryDefinition)
+	{
+		DerivedSchemaDefinition = PreparedServiceChanges->DerivedSchemaDefinition;
+		DerivedSchemaCategoryDefinition = PreparedServiceChanges->DerivedSchemaCategoryDefinition;
+	}
+
+	// Move out client changes for returning to the client.
+	ClientChanges = MoveTemp(PreparedServiceChanges->ClientChanges);
+
+	// Reset pending state.
+	ResetPreparedChanges();
+
+	// Return client changes.
+	return FSchemaCategoryInstanceCommitServiceSnapshot::Result({ MoveTemp(ClientChanges) });
 }
 
-bool FSchemaCategoryInstance::IsValid() const
+bool FSchemaCategoryInstanceBase::IsValid() const
 {
 	return BaseSchemaCategoryDefinition != nullptr;
 }
 
-bool FSchemaCategoryInstance::IsValidBaseAttributeData(const FSchemaAttributeId& Id, const FSchemaVariant& Data)
+bool FSchemaCategoryInstanceBase::VerifyBaseAttributeData(
+	const FSchemaAttributeId& Id,
+	const FSchemaVariant& Data,
+	FSchemaServiceAttributeId& OutSchemaServiceAttributeId,
+	ESchemaServiceAttributeFlags& OutSchemaServiceAttributeFlags)
 {
 	if (ensure(IsValid()))
 	{
@@ -1125,18 +1308,39 @@ bool FSchemaCategoryInstance::IsValidBaseAttributeData(const FSchemaAttributeId&
 		{
 			if (AttributeDefinition->Type == Data.GetType())
 			{
-				if (AttributeDefinition->Type == ESchemaAttributeType::String)
+				const bool bIsDataSizeValid = AttributeDefinition->Type != ESchemaAttributeType::String || Data.GetString().Len() <= AttributeDefinition->MaxSize;
+				if (bIsDataSizeValid)
 				{
-					return Data.GetString().Len() <= AttributeDefinition->MaxSize;
+					OutSchemaServiceAttributeId = AttributeDefinition->ServiceAttributeId;
+
+					if (AttributeDefinition->ServiceAttributeId != FSchemaServiceAttributeId())
+					{
+						const FSchemaServiceAttributeDefinition* ServiceAttributeDefinition = BaseSchemaCategoryDefinition->ServiceAttributeDefinitions.Find(AttributeDefinition->ServiceAttributeId);
+						if (ensure(ServiceAttributeDefinition))
+						{
+							OutSchemaServiceAttributeFlags = ServiceAttributeDefinition->Flags;
+						}
+						else
+						{
+							UE_LOG(LogOnlineSchema, Warning, TEXT("[VerifyBaseAttributeData] %s.%s.%s: Unable to find service attribute definition for service attribute %s."),
+								*BaseSchemaCategoryDefinition->Id.ToString().ToLower(),
+								*BaseSchemaCategoryDefinition->Id.ToString().ToLower(),
+								*AttributeDefinition->Id.ToString().ToLower(),
+								*AttributeDefinition->ServiceAttributeId.ToString().ToLower());
+							return false;
+						}
+					}
+
+					return true;
 				}
 				else
 				{
-					return true;
+					return false;
 				}
 			}
 			else
 			{
-				UE_LOG(LogOnlineSchema, Verbose, TEXT("[IsValidBaseAttributeData] Schema attribute %s.%s.%s set with invalid type %s. Expected type %s"),
+				UE_LOG(LogOnlineSchema, Verbose, TEXT("[VerifyBaseAttributeData] Schema attribute %s.%s.%s set with invalid type %s. Expected type %s"),
 					*BaseSchemaDefinition->Id.ToString().ToLower(),
 					*BaseSchemaCategoryDefinition->Id.ToString().ToLower(),
 					*AttributeDefinition->Id.ToString().ToLower(),
@@ -1147,7 +1351,7 @@ bool FSchemaCategoryInstance::IsValidBaseAttributeData(const FSchemaAttributeId&
 		}
 		else
 		{
-			UE_LOG(LogOnlineSchema, Verbose, TEXT("[IsValidBaseAttributeData] Attribute %s not found in schema definition %s.%s"),
+			UE_LOG(LogOnlineSchema, Verbose, TEXT("[VerifyBaseAttributeData] Attribute %s not found in schema definition %s.%s"),
 				*Id.ToString().ToLower(),
 				*BaseSchemaDefinition->Id.ToString().ToLower(),
 				*BaseSchemaCategoryDefinition->Id.ToString().ToLower());
@@ -1156,36 +1360,46 @@ bool FSchemaCategoryInstance::IsValidBaseAttributeData(const FSchemaAttributeId&
 	}
 	else
 	{
-		UE_LOG(LogOnlineSchema, Verbose, TEXT("[IsValidBaseAttributeData] Unable to set attribute %s. Schema is not valid"),
+		UE_LOG(LogOnlineSchema, Verbose, TEXT("[VerifyBaseAttributeData] Unable to set attribute %s. Schema is not valid"),
 			*Id.ToString().ToLower());
 		return false;
 	}
 }
 
-bool FSchemaCategoryInstance::InitializeSchemaDefinition(
-	const FSchemaId& SchemaId,
+const TMap<FSchemaAttributeId, FSchemaVariant>& FSchemaCategoryInstanceBase::GetClientSnapshot() const
+{
+	return const_cast<FSchemaCategoryInstanceBase*>(this)->GetMutableClientSnapshot();
+}
+
+bool FSchemaCategoryInstanceBase::InitializeSchemaDefinition(
+	int64 SchemaCompatibilityId,
+	const FSchemaDefinition* OptionalBaseDefinition,
 	const FSchemaCategoryId& CategoryId,
 	TSharedPtr<const FSchemaDefinition>* OutSchemaDefinition,
-	const FSchemaCategoryDefinition** OutSchemaCategoryDefinition)
+	const FSchemaCategoryDefinition** OutSchemaCategoryDefinition) const
 {
-	check(OutSchemaDefinition);
-	check(OutSchemaCategoryDefinition);
+	if (TSharedPtr<const FSchemaDefinition> SchemaDefinition = SchemaRegistry->GetDefinition(SchemaCompatibilityId))
+	{
+		return InitializeSchemaDefinition(SchemaDefinition.ToSharedRef(), OptionalBaseDefinition, CategoryId, OutSchemaDefinition, OutSchemaCategoryDefinition);
+	}
+	else
+	{
+		UE_LOG(LogOnlineSchema, Error, TEXT("[InitializeSchemaDefinition] init error: Unable to find definition for SchemaCompatibilityId 0x%08" INT64_X_FMT "."),
+			SchemaCompatibilityId);
+		return false;
+	}
+}
 
+bool FSchemaCategoryInstanceBase::InitializeSchemaDefinition(
+	const FSchemaId& SchemaId,
+	const FSchemaDefinition* OptionalBaseDefinition,
+	const FSchemaCategoryId& CategoryId,
+	TSharedPtr<const FSchemaDefinition>* OutSchemaDefinition,
+	const FSchemaCategoryDefinition** OutSchemaCategoryDefinition) const
+{
 	if (TSharedPtr<const FSchemaDefinition> SchemaDefinition = SchemaRegistry->GetDefinition(SchemaId))
 	{
-		if (const FSchemaCategoryDefinition* SchemaCategoryDefinition = SchemaDefinition->Categories.Find(CategoryId))
-		{
-			*OutSchemaDefinition = MoveTemp(SchemaDefinition);
-			*OutSchemaCategoryDefinition = SchemaCategoryDefinition;
-			return true;
-		}
-		else
-		{
-			UE_LOG(LogOnlineSchema, Error, TEXT("[InitializeSchemaDefinition] init error: Unable to find category %s in schema %s"),
-				*CategoryId.ToString().ToLower(),
-				*SchemaId.ToString().ToLower());
-			return false;
-		}
+		return InitializeSchemaDefinition(SchemaDefinition.ToSharedRef(), OptionalBaseDefinition, CategoryId, OutSchemaDefinition, OutSchemaCategoryDefinition);
 	}
 	else
 	{
@@ -1195,7 +1409,41 @@ bool FSchemaCategoryInstance::InitializeSchemaDefinition(
 	}
 }
 
-bool FSchemaCategoryInstance::GetSerializationSchema(
+bool FSchemaCategoryInstanceBase::InitializeSchemaDefinition(
+	const TSharedRef<const FSchemaDefinition>& NewDefinition,
+	const FSchemaDefinition* OptionalBaseDefinition,
+	const FSchemaCategoryId& CategoryId,
+	TSharedPtr<const FSchemaDefinition>* OutSchemaDefinition,
+	const FSchemaCategoryDefinition** OutSchemaCategoryDefinition) const
+{
+	check(OutSchemaDefinition);
+	check(OutSchemaCategoryDefinition);
+
+	// Enforce parent <-> child relationship if needed.
+	if (OptionalBaseDefinition && !SchemaRegistry->IsSchemaChildOf(NewDefinition->Id, OptionalBaseDefinition->Id))
+	{
+		UE_LOG(LogOnlineSchema, Warning, TEXT("[InitializeSchemaDefinition] Invalid schema definition. New schema %s is not a child of %s."),
+			*NewDefinition->Id.ToString().ToLower(),
+			*OptionalBaseDefinition->Id.ToString().ToLower());
+		return false;
+	}
+
+	// Check that category exists in incoming schema.
+	const FSchemaCategoryDefinition* SchemaCategoryDefinition = NewDefinition->Categories.Find(CategoryId);
+	if (SchemaCategoryDefinition == nullptr)
+	{
+		UE_LOG(LogOnlineSchema, Error, TEXT("[InitializeSchemaDefinition] init error: Unable to find category %s in schema %s"),
+			*CategoryId.ToString().ToLower(),
+			*NewDefinition->Id.ToString().ToLower());
+		return false;
+	}
+
+	*OutSchemaDefinition = NewDefinition;
+	*OutSchemaCategoryDefinition = SchemaCategoryDefinition;
+	return true;
+}
+
+bool FSchemaCategoryInstanceBase::GetSerializationSchema(
 	const FSchemaDefinition** OutSchemaDefinition,
 	const FSchemaCategoryDefinition** OutSchemaCategoryDefinition) const
 {
@@ -1223,6 +1471,14 @@ bool FSchemaCategoryInstance::GetSerializationSchema(
 	}
 }
 
+void FSchemaCategoryInstanceBase::ResetPreparedChanges() const
+{
+	PreparedClientChanges.Reset();
+	PreparedServiceChanges.Reset();
+}
+
+/* Private */ }
+
 TSharedPtr<const FSchemaDefinition> FSchemaRegistry::GetDefinition(const FSchemaId& SchemaId) const
 {
 	const TSharedRef<const FSchemaDefinition>* Definition = SchemaDefinitionsById.Find(SchemaId);
diff --git a/Engine/Plugins/Online/Experimental/OnlineServices/Source/OnlineServicesInterface/Public/Online/Lobbies.h b/Engine/Plugins/Online/Experimental/OnlineServices/Source/OnlineServicesInterface/Public/Online/Lobbies.h
--- a/Engine/Plugins/Online/Experimental/OnlineServices/Source/OnlineServicesInterface/Public/Online/Lobbies.h
+++ b/Engine/Plugins/Online/Experimental/OnlineServices/Source/OnlineServicesInterface/Public/Online/Lobbies.h
@@ -10,13 +10,6 @@
 
 namespace UE::Online {
 
-// TODO: Aliases will be removed once full Schema support is added to Lobbies
-using FLobbySchemaId = FSchemaId;
-using FLobbyAttributeId = FSchemaAttributeId;
-using FLobbyVariant = FSchemaVariant;
-using ELobbyComparisonOp = ESchemaAttributeComparisonOp;
-using ELobbyAttributeVisibility = ESchemaAttributeVisibility;
-
 enum class ELobbyJoinPolicy : uint8
 {
 	/** 
@@ -59,7 +52,8 @@ struct FLobbyMember
 	FAccountId AccountId;
 	FAccountId PlatformAccountId;
 	FString PlatformDisplayName;
-	TMap<FLobbyAttributeId, FLobbyVariant> Attributes;
+	TMap<FSchemaAttributeId, FSchemaVariant> Attributes;
+	bool bIsLocalMember = false;
 };
 
 struct FLobby
@@ -67,23 +61,23 @@ struct FLobby
 	FLobbyId LobbyId;
 	FAccountId OwnerAccountId;
 	FName LocalName;
-	FName SchemaName;
+	FSchemaId SchemaId;
 	int32 MaxMembers;
 	ELobbyJoinPolicy JoinPolicy;
-	TMap<FLobbyAttributeId, FLobbyVariant> Attributes;
+	TMap<FSchemaAttributeId, FSchemaVariant> Attributes;
 	TMap<FAccountId, TSharedRef<const FLobbyMember>> Members;
 };
 
 struct FFindLobbySearchFilter
 {
 	/** The name of the attribute to be compared. */
-	FLobbyAttributeId AttributeName;
+	FSchemaAttributeId AttributeName;
 
 	/** The type of comparison to perform. */
-	ELobbyComparisonOp ComparisonOp;
+	ESchemaAttributeComparisonOp ComparisonOp;
 
 	/** Value to use when comparing the attribute. */
-	FLobbyVariant ComparisonValue;
+	FSchemaVariant ComparisonValue;
 };
 
 ONLINESERVICESINTERFACE_API void SortLobbies(const TArray<FFindLobbySearchFilter>& Filters, TArray<TSharedRef<const FLobby>>& Lobbies);
@@ -102,7 +96,7 @@ struct FCreateLobby
 		FName LocalName;
 
 		/** The schema which will be applied to the lobby. */
-		FName SchemaName;
+		FSchemaId SchemaId;
 
 		/** The maximum number of members who can fit in the lobby. */
 		int32 MaxMembers = 0;
@@ -111,10 +105,10 @@ struct FCreateLobby
 		ELobbyJoinPolicy JoinPolicy = ELobbyJoinPolicy::InvitationOnly;
 
 		/** Initial attributes. */
-		TMap<FLobbyAttributeId, FLobbyVariant> Attributes;
+		TMap<FSchemaAttributeId, FSchemaVariant> Attributes;
 
 		/** Initial user attributes. */
-		TMap<FLobbyAttributeId, FLobbyVariant> UserAttributes;
+		TMap<FSchemaAttributeId, FSchemaVariant> UserAttributes;
 	};
 
 	/** Output struct for Lobbies::CreateLobby */
@@ -190,7 +184,7 @@ struct FJoinLobby
 		FLobbyId LobbyId;
 
 		/** Initial user attributes. */
-		TMap<FLobbyAttributeId, FLobbyVariant> UserAttributes;
+		TMap<FSchemaAttributeId, FSchemaVariant> UserAttributes;
 	};
 
 	/** Output struct for Lobbies::JoinLobby */
@@ -362,10 +356,10 @@ struct FModifyLobbyAttributes
 		FLobbyId LobbyId;
 
 		/** New or changed lobby attributes. */
-		TMap<FLobbyAttributeId, FLobbyVariant> MutatedAttributes;
+		TMap<FSchemaAttributeId, FSchemaVariant> UpdatedAttributes;
 
 		/** Attributes to be cleared. */
-		TSet<FLobbyAttributeId> ClearedAttributes;
+		TSet<FSchemaAttributeId> RemovedAttributes;
 	};
 
 	/** Output struct for Lobbies::ModifyLobbyAttributes */
@@ -388,10 +382,10 @@ struct FModifyLobbyMemberAttributes
 		FLobbyId LobbyId;
 
 		/** New or changed lobby attributes. */
-		TMap<FLobbyAttributeId, FLobbyVariant> MutatedAttributes;
+		TMap<FSchemaAttributeId, FSchemaVariant> UpdatedAttributes;
 
 		/** Attributes to be cleared. */
-		TSet<FLobbyAttributeId> ClearedAttributes;
+		TSet<FSchemaAttributeId> RemovedAttributes;
 	};
 
 	/** Output struct for Lobbies::ModifyLobbyMemberAttributes */
@@ -487,6 +481,10 @@ struct FLobbySchemaChanged
 {
 	/** Lobby data access. */
 	TSharedRef<const FLobby> Lobby;
+	/** The previous schema id. */
+	FSchemaId PreviousSchemaId;
+	/** The new schema id. */
+	FSchemaId NewSchemaId;
 };
 
 /** Struct for LobbyAttributesChanged event */
@@ -495,8 +493,14 @@ struct FLobbyAttributesChanged
 	/** Lobby data access. */
 	TSharedRef<const FLobby> Lobby;
 
-	/** Attribute keys which changed. */
-	TSet<FLobbyAttributeId> ChangedAttributes;
+	/** Added attributes. */
+	TMap<FSchemaAttributeId, FSchemaVariant> AddedAttributes;
+
+	/** Changed attributes with their old and new values. */
+	TMap<FSchemaAttributeId, TPair<FSchemaVariant, FSchemaVariant>> ChangedAttributes;
+
+	/** Removed attribute ids. */
+	TSet<FSchemaAttributeId> RemovedAttributes;
 
 };
 
@@ -509,8 +513,14 @@ struct FLobbyMemberAttributesChanged
 	/** Member data access. */
 	TSharedRef<const FLobbyMember> Member;
 
-	/** Attribute keys which changed. */
-	TSet<FLobbyAttributeId> ChangedAttributes;
+	/** Added attributes. */
+	TMap<FSchemaAttributeId, FSchemaVariant> AddedAttributes;
+
+	/** Changed attributes with their old and new values. */
+	TMap<FSchemaAttributeId, TPair<FSchemaVariant, FSchemaVariant>> ChangedAttributes;
+
+	/** Removed attribute ids. */
+	TSet<FSchemaAttributeId> RemovedAttributes;
 };
 
 /** Struct for LobbyInvitationAdded event */
@@ -557,7 +567,7 @@ struct FUILobbyJoinRequested
 	/** The local user associated with the join request. */
 	FAccountId LocalAccountId;
 
-	/** The lobby the local user requested to join, or the online error if there was a failure */
+	/** The lobby if the local user requested to join, or the online error if there was a failure */
 	TResult<TSharedRef<const FLobby>, FOnlineError> Result;
 
 	/** Join request source */
@@ -789,20 +799,35 @@ public:
 };
 
 namespace Meta {
-// TODO: Move to Lobbies_Meta.inl file?
+
+BEGIN_ONLINE_STRUCT_META(FLobbyMember)
+	ONLINE_STRUCT_FIELD(FLobbyMember, AccountId),
+	ONLINE_STRUCT_FIELD(FLobbyMember, PlatformAccountId),
+	ONLINE_STRUCT_FIELD(FLobbyMember, PlatformDisplayName),
+	ONLINE_STRUCT_FIELD(FLobbyMember, Attributes)
+END_ONLINE_STRUCT_META()
 
 BEGIN_ONLINE_STRUCT_META(FLobby)
 	ONLINE_STRUCT_FIELD(FLobby, LobbyId),
+	ONLINE_STRUCT_FIELD(FLobby, OwnerAccountId),
 	ONLINE_STRUCT_FIELD(FLobby, LocalName),
-	ONLINE_STRUCT_FIELD(FLobby, SchemaName),
+	ONLINE_STRUCT_FIELD(FLobby, SchemaId),
 	ONLINE_STRUCT_FIELD(FLobby, MaxMembers),
-	ONLINE_STRUCT_FIELD(FLobby, JoinPolicy)
+	ONLINE_STRUCT_FIELD(FLobby, JoinPolicy),
+	ONLINE_STRUCT_FIELD(FLobby, Attributes),
+	ONLINE_STRUCT_FIELD(FLobby, Members)
+END_ONLINE_STRUCT_META()
+
+BEGIN_ONLINE_STRUCT_META(FFindLobbySearchFilter)
+	ONLINE_STRUCT_FIELD(FFindLobbySearchFilter, AttributeName),
+	ONLINE_STRUCT_FIELD(FFindLobbySearchFilter, ComparisonOp),
+	ONLINE_STRUCT_FIELD(FFindLobbySearchFilter, ComparisonValue)
 END_ONLINE_STRUCT_META()
 
 BEGIN_ONLINE_STRUCT_META(FCreateLobby::Params)
 	ONLINE_STRUCT_FIELD(FCreateLobby::Params, LocalAccountId),
 	ONLINE_STRUCT_FIELD(FCreateLobby::Params, LocalName),
-	ONLINE_STRUCT_FIELD(FCreateLobby::Params, SchemaName),
+	ONLINE_STRUCT_FIELD(FCreateLobby::Params, SchemaId),
 	ONLINE_STRUCT_FIELD(FCreateLobby::Params, MaxMembers),
 	ONLINE_STRUCT_FIELD(FCreateLobby::Params, JoinPolicy),
 	ONLINE_STRUCT_FIELD(FCreateLobby::Params, Attributes),
@@ -814,9 +839,15 @@ BEGIN_ONLINE_STRUCT_META(FCreateLobby::Result)
 END_ONLINE_STRUCT_META()
 
 BEGIN_ONLINE_STRUCT_META(FFindLobbies::Params)
+	ONLINE_STRUCT_FIELD(FFindLobbies::Params, LocalAccountId),
+	ONLINE_STRUCT_FIELD(FFindLobbies::Params, MaxResults),
+	ONLINE_STRUCT_FIELD(FFindLobbies::Params, Filters),
+	ONLINE_STRUCT_FIELD(FFindLobbies::Params, TargetUser),
+	ONLINE_STRUCT_FIELD(FFindLobbies::Params, LobbyId)
 END_ONLINE_STRUCT_META()
 
 BEGIN_ONLINE_STRUCT_META(FFindLobbies::Result)
+	ONLINE_STRUCT_FIELD(FFindLobbies::Result, Lobbies)
 END_ONLINE_STRUCT_META()
 
 BEGIN_ONLINE_STRUCT_META(FRestoreLobbies::Params)
@@ -826,9 +857,14 @@ BEGIN_ONLINE_STRUCT_META(FRestoreLobbies::Result)
 END_ONLINE_STRUCT_META()
 
 BEGIN_ONLINE_STRUCT_META(FJoinLobby::Params)
+	ONLINE_STRUCT_FIELD(FJoinLobby::Params, LocalAccountId),
+	ONLINE_STRUCT_FIELD(FJoinLobby::Params, LocalName),
+	ONLINE_STRUCT_FIELD(FJoinLobby::Params, LobbyId),
+	ONLINE_STRUCT_FIELD(FJoinLobby::Params, UserAttributes)
 END_ONLINE_STRUCT_META()
 
 BEGIN_ONLINE_STRUCT_META(FJoinLobby::Result)
+	ONLINE_STRUCT_FIELD(FJoinLobby::Result, Lobby)
 END_ONLINE_STRUCT_META()
 
 BEGIN_ONLINE_STRUCT_META(FLeaveLobby::Params)
@@ -840,24 +876,35 @@ BEGIN_ONLINE_STRUCT_META(FLeaveLobby::Result)
 END_ONLINE_STRUCT_META()
 
 BEGIN_ONLINE_STRUCT_META(FInviteLobbyMember::Params)
+	ONLINE_STRUCT_FIELD(FInviteLobbyMember::Params, LocalAccountId),
+	ONLINE_STRUCT_FIELD(FInviteLobbyMember::Params, LobbyId),
+	ONLINE_STRUCT_FIELD(FInviteLobbyMember::Params, TargetAccountId)
 END_ONLINE_STRUCT_META()
 
 BEGIN_ONLINE_STRUCT_META(FInviteLobbyMember::Result)
 END_ONLINE_STRUCT_META()
 
 BEGIN_ONLINE_STRUCT_META(FDeclineLobbyInvitation::Params)
+	ONLINE_STRUCT_FIELD(FDeclineLobbyInvitation::Params, LocalAccountId),
+	ONLINE_STRUCT_FIELD(FDeclineLobbyInvitation::Params, LobbyId)
 END_ONLINE_STRUCT_META()
 
 BEGIN_ONLINE_STRUCT_META(FDeclineLobbyInvitation::Result)
 END_ONLINE_STRUCT_META()
 
 BEGIN_ONLINE_STRUCT_META(FKickLobbyMember::Params)
+	ONLINE_STRUCT_FIELD(FKickLobbyMember::Params, LocalAccountId),
+	ONLINE_STRUCT_FIELD(FKickLobbyMember::Params, LobbyId),
+	ONLINE_STRUCT_FIELD(FKickLobbyMember::Params, TargetAccountId)
 END_ONLINE_STRUCT_META()
 
 BEGIN_ONLINE_STRUCT_META(FKickLobbyMember::Result)
 END_ONLINE_STRUCT_META()
 
 BEGIN_ONLINE_STRUCT_META(FPromoteLobbyMember::Params)
+	ONLINE_STRUCT_FIELD(FPromoteLobbyMember::Params, LocalAccountId),
+	ONLINE_STRUCT_FIELD(FPromoteLobbyMember::Params, LobbyId),
+	ONLINE_STRUCT_FIELD(FPromoteLobbyMember::Params, TargetAccountId)
 END_ONLINE_STRUCT_META()
 
 BEGIN_ONLINE_STRUCT_META(FPromoteLobbyMember::Result)
@@ -870,18 +917,29 @@ BEGIN_ONLINE_STRUCT_META(FModifyLobbySchema::Result)
 END_ONLINE_STRUCT_META()
 
 BEGIN_ONLINE_STRUCT_META(FModifyLobbyJoinPolicy::Params)
+	ONLINE_STRUCT_FIELD(FModifyLobbyJoinPolicy::Params, LocalAccountId),
+	ONLINE_STRUCT_FIELD(FModifyLobbyJoinPolicy::Params, LobbyId),
+	ONLINE_STRUCT_FIELD(FModifyLobbyJoinPolicy::Params, JoinPolicy)
 END_ONLINE_STRUCT_META()
 
 BEGIN_ONLINE_STRUCT_META(FModifyLobbyJoinPolicy::Result)
 END_ONLINE_STRUCT_META()
 
 BEGIN_ONLINE_STRUCT_META(FModifyLobbyAttributes::Params)
+	ONLINE_STRUCT_FIELD(FModifyLobbyAttributes::Params, LocalAccountId),
+	ONLINE_STRUCT_FIELD(FModifyLobbyAttributes::Params, LobbyId),
+	ONLINE_STRUCT_FIELD(FModifyLobbyAttributes::Params, UpdatedAttributes),
+	ONLINE_STRUCT_FIELD(FModifyLobbyAttributes::Params, RemovedAttributes)
 END_ONLINE_STRUCT_META()
 
 BEGIN_ONLINE_STRUCT_META(FModifyLobbyAttributes::Result)
 END_ONLINE_STRUCT_META()
 
 BEGIN_ONLINE_STRUCT_META(FModifyLobbyMemberAttributes::Params)
+	ONLINE_STRUCT_FIELD(FModifyLobbyMemberAttributes::Params, LocalAccountId),
+	ONLINE_STRUCT_FIELD(FModifyLobbyMemberAttributes::Params, LobbyId),
+	ONLINE_STRUCT_FIELD(FModifyLobbyMemberAttributes::Params, UpdatedAttributes),
+	ONLINE_STRUCT_FIELD(FModifyLobbyMemberAttributes::Params, RemovedAttributes)
 END_ONLINE_STRUCT_META()
 
 BEGIN_ONLINE_STRUCT_META(FModifyLobbyMemberAttributes::Result)
@@ -926,15 +984,25 @@ BEGIN_ONLINE_STRUCT_META(FLobbyLeaderChanged)
 	ONLINE_STRUCT_FIELD(FLobbyLeaderChanged, Leader)
 END_ONLINE_STRUCT_META()
 
+BEGIN_ONLINE_STRUCT_META(FLobbySchemaChanged)
+	ONLINE_STRUCT_FIELD(FLobbySchemaChanged, Lobby),
+	ONLINE_STRUCT_FIELD(FLobbySchemaChanged, PreviousSchemaId),
+	ONLINE_STRUCT_FIELD(FLobbySchemaChanged, NewSchemaId)
+END_ONLINE_STRUCT_META()
+
 BEGIN_ONLINE_STRUCT_META(FLobbyAttributesChanged)
 	ONLINE_STRUCT_FIELD(FLobbyAttributesChanged, Lobby),
-	ONLINE_STRUCT_FIELD(FLobbyAttributesChanged, ChangedAttributes)
+	ONLINE_STRUCT_FIELD(FLobbyAttributesChanged, AddedAttributes),
+	ONLINE_STRUCT_FIELD(FLobbyAttributesChanged, ChangedAttributes),
+	ONLINE_STRUCT_FIELD(FLobbyAttributesChanged, RemovedAttributes)
 END_ONLINE_STRUCT_META()
 
 BEGIN_ONLINE_STRUCT_META(FLobbyMemberAttributesChanged)
 	ONLINE_STRUCT_FIELD(FLobbyMemberAttributesChanged, Lobby),
 	ONLINE_STRUCT_FIELD(FLobbyMemberAttributesChanged, Member),
-	ONLINE_STRUCT_FIELD(FLobbyMemberAttributesChanged, ChangedAttributes)
+	ONLINE_STRUCT_FIELD(FLobbyMemberAttributesChanged, AddedAttributes),
+	ONLINE_STRUCT_FIELD(FLobbyMemberAttributesChanged, ChangedAttributes),
+	ONLINE_STRUCT_FIELD(FLobbyMemberAttributesChanged, RemovedAttributes)
 END_ONLINE_STRUCT_META()
 
 BEGIN_ONLINE_STRUCT_META(FLobbyInvitationAdded)
@@ -951,7 +1019,8 @@ END_ONLINE_STRUCT_META()
 
 BEGIN_ONLINE_STRUCT_META(FUILobbyJoinRequested)
 	ONLINE_STRUCT_FIELD(FUILobbyJoinRequested, LocalAccountId),
-	ONLINE_STRUCT_FIELD(FUILobbyJoinRequested, Result)
+	ONLINE_STRUCT_FIELD(FUILobbyJoinRequested, Result),
+	ONLINE_STRUCT_FIELD(FUILobbyJoinRequested, JoinRequestedSource)
 END_ONLINE_STRUCT_META()
 
 /* Meta*/ }
diff --git a/Engine/Plugins/Online/Experimental/OnlineServices/Source/OnlineServicesInterface/Public/Online/Schema.h b/Engine/Plugins/Online/Experimental/OnlineServices/Source/OnlineServicesInterface/Public/Online/Schema.h
--- a/Engine/Plugins/Online/Experimental/OnlineServices/Source/OnlineServicesInterface/Public/Online/Schema.h
+++ b/Engine/Plugins/Online/Experimental/OnlineServices/Source/OnlineServicesInterface/Public/Online/Schema.h
@@ -9,7 +9,9 @@ namespace UE::Online {
 
 class FSchemaRegistry;
 
-class ONLINESERVICESINTERFACE_API FSchemaCategoryInstance
+namespace Private {
+
+class ONLINESERVICESINTERFACE_API FSchemaCategoryInstanceBase
 {
 public:
 	// DerivedSchemaId may remain unset in two situations:
@@ -17,17 +19,26 @@ public:
 	//    the base schema is required to have an attribute flagged as SchemaCompatibilityId so that
 	//    the derived schema can be discovered.
 	// 2. Schema swapping is not enabled. All attributes are defined in the base schema.
-	FSchemaCategoryInstance(
+	FSchemaCategoryInstanceBase(
 		const FSchemaId& DerivedSchemaId,
 		const FSchemaId& BaseSchemaId,
 		const FSchemaCategoryId& CategoryId,
 		const TSharedRef<const FSchemaRegistry>& SchemaRegistry);
 
-#if 0
-	TOnlineResult<FSchemaCategoryInstanceChangeSchema> ChangeSchema(FSchemaCategoryInstanceChangeSchema::Params&& Params);
-#endif
-	TOnlineResult<FSchemaCategoryInstanceApplyApplicationDelta> ApplyApplicationDelta(FSchemaCategoryInstanceApplyApplicationDelta::Params&& Params);
-	TOnlineResult<FSchemaCategoryInstanceApplyServiceSnapshot> ApplyServiceSnapshot(FSchemaCategoryInstanceApplyServiceSnapshot::Params&& Params);
+	virtual ~FSchemaCategoryInstanceBase() = default;
+
+	// Two phase commit to the service.
+	// Prepare data to be written to the service from an client delta of attributes.
+	TOnlineResult<FSchemaCategoryInstancePrepareClientChanges> PrepareClientChanges(FSchemaCategoryInstancePrepareClientChanges::Params&& Params) const;
+	// Commit the client data once the service data has been successfully written.
+	FSchemaCategoryInstanceCommitClientChanges::Result CommitClientChanges();
+
+	// Two phase commit from the service.
+	// Prepare data to be written to the client from a service snapshot of attributes.
+	TOnlineResult<FSchemaCategoryInstancePrepareServiceSnapshot> PrepareServiceSnapshot(FSchemaCategoryInstancePrepareServiceSnapshot::Params&& Params) const;
+	// Commit a service snapshot containing all known service attributes.
+	FSchemaCategoryInstanceCommitServiceSnapshot::Result CommitServiceSnapshot();
+
 	// Todo: ApplyServiceDelta
 
 	TSharedPtr<const FSchemaDefinition> GetDerivedDefinition() const { return DerivedSchemaDefinition; };
@@ -36,21 +47,62 @@ public:
 	// Check whether the schema is valid.
 	bool IsValid() const;
 
-	// Check whether an application attribute is valid in the base schema.
-	// Intended to be used when a schema is not backed by a service.
-	bool IsValidBaseAttributeData(const FSchemaAttributeId& Id, const FSchemaVariant& Data);
+	// Check whether an client attribute is valid in the base schema.
+	// Intended to be used when attribute information needs to be verified without translating to the service.
+	bool VerifyBaseAttributeData(
+		const FSchemaAttributeId& Id,
+		const FSchemaVariant& Data,
+		FSchemaServiceAttributeId& OutSchemaServiceAttributeId,
+		ESchemaServiceAttributeFlags& OutSchemaServiceAttributeFlags);
+
+protected:
+	const TMap<FSchemaAttributeId, FSchemaVariant>& GetClientSnapshot() const;
+	virtual TMap<FSchemaAttributeId, FSchemaVariant>& GetMutableClientSnapshot() = 0;
 
 private:
+	bool InitializeSchemaDefinition(
+		int64 SchemaCompatibilityId,
+		const FSchemaDefinition* OptionalBaseDefinition,
+		const FSchemaCategoryId& CategoryId,
+		TSharedPtr<const FSchemaDefinition>* OutSchemaDefinition,
+		const FSchemaCategoryDefinition** OutSchemaCategoryDefinition) const;
+
 	bool InitializeSchemaDefinition(
 		const FSchemaId& SchemaId,
+		const FSchemaDefinition* OptionalBaseDefinition,
 		const FSchemaCategoryId& CategoryId,
 		TSharedPtr<const FSchemaDefinition>* OutSchemaDefinition,
-		const FSchemaCategoryDefinition** OutSchemaCategoryDefinition);
+		const FSchemaCategoryDefinition** OutSchemaCategoryDefinition) const;
+
+	bool InitializeSchemaDefinition(
+		const TSharedRef<const FSchemaDefinition>& NewDefinition,
+		const FSchemaDefinition* OptionalBaseDefinition,
+		const FSchemaCategoryId& CategoryId,
+		TSharedPtr<const FSchemaDefinition>* OutSchemaDefinition,
+		const FSchemaCategoryDefinition** OutSchemaCategoryDefinition) const;
 
 	bool GetSerializationSchema(
 		const FSchemaDefinition** OutSchemaDefinition,
 		const FSchemaCategoryDefinition** OutSchemaCategoryDefinition) const;
 
+	void ResetPreparedChanges() const;
+
+	struct FPreparedClientChanges
+	{
+		FSchemaServiceClientChanges ClientChanges;
+		int64 SchemaCompatibilityId = 0;
+		TSharedPtr<const FSchemaDefinition> DerivedSchemaDefinition;
+		const FSchemaCategoryDefinition* DerivedSchemaCategoryDefinition = nullptr;
+	};
+
+	struct FPreparedServiceChanges
+	{
+		FSchemaServiceClientChanges ClientChanges;
+		TSharedPtr<const FSchemaDefinition> DerivedSchemaDefinition;
+		const FSchemaCategoryDefinition* DerivedSchemaCategoryDefinition = nullptr;
+		TMap<FSchemaAttributeId, FSchemaVariant> ClientDataSnapshot;
+	};
+
 	TSharedRef<const FSchemaRegistry> SchemaRegistry;
 	TSharedPtr<const FSchemaDefinition> DerivedSchemaDefinition;
 	TSharedPtr<const FSchemaDefinition> BaseSchemaDefinition;
@@ -58,10 +110,59 @@ private:
 	const FSchemaCategoryDefinition* BaseSchemaCategoryDefinition = nullptr;
 	int64 LastSentSchemaCompatibilityId = 0;
 
-	// Snapshot of application data as provided by the service.
-	TMap<FSchemaAttributeId, FSchemaVariant> ApplicationDataServiceSnapshot;
+	// Prepared changes are mutable so that the prepare methods can be const. Setting as mutable
+	// is done to prevent modifying state outside of the prepared changes.
+	mutable TOptional<FPreparedClientChanges> PreparedClientChanges;
+	mutable TOptional<FPreparedServiceChanges> PreparedServiceChanges;
 };
 
+/* Private */ }
+
+/** Schema category instance attribute accessor for providing the client attributes internally. */
+class FSchemaCategoryInstanceInternalSnapshotAccessor
+{
+public:
+	TMap<FSchemaAttributeId, FSchemaVariant>& GetMutableClientSnapshot()
+	{
+		return ClientSnapshot;
+	}
+private:
+
+	TMap<FSchemaAttributeId, FSchemaVariant> ClientSnapshot;
+};
+
+/**
+  * Schema category instance class templated by access to client attributes.
+  * 
+  * The passed in accessor is an interface which must provide a definition for GetMutableClientSnapshot.
+  */
+template <typename FSnapshotAccessor>
+class TSchemaCategoryInstance final : public Private::FSchemaCategoryInstanceBase
+{
+public:
+	TSchemaCategoryInstance(
+		const FSchemaId& DerivedSchemaId,
+		const FSchemaId& BaseSchemaId,
+		const FSchemaCategoryId& CategoryId,
+		const TSharedRef<const FSchemaRegistry>& SchemaRegistry,
+		FSnapshotAccessor&& InSnapshotAccessor = FSnapshotAccessor())
+		: Private::FSchemaCategoryInstanceBase(DerivedSchemaId, BaseSchemaId, CategoryId, SchemaRegistry)
+		, SnapshotAccessor(MoveTemp(InSnapshotAccessor))
+	{
+	}
+
+private:
+	virtual TMap<FSchemaAttributeId, FSchemaVariant>& GetMutableClientSnapshot() override
+	{
+		return SnapshotAccessor.GetMutableClientSnapshot();
+	}
+
+	FSnapshotAccessor SnapshotAccessor;
+};
+
+/** Default implementation with both translation and client attribute snapshot data contained within the category instance. */
+using FSchemaCategoryInstance = TSchemaCategoryInstance<FSchemaCategoryInstanceInternalSnapshotAccessor>;
+
 class ONLINESERVICESINTERFACE_API FSchemaRegistry
 {
 public:
diff --git a/Engine/Plugins/Online/Experimental/OnlineServices/Source/OnlineServicesInterface/Public/Online/SchemaTypes.h b/Engine/Plugins/Online/Experimental/OnlineServices/Source/OnlineServicesInterface/Public/Online/SchemaTypes.h
--- a/Engine/Plugins/Online/Experimental/OnlineServices/Source/OnlineServicesInterface/Public/Online/SchemaTypes.h
+++ b/Engine/Plugins/Online/Experimental/OnlineServices/Source/OnlineServicesInterface/Public/Online/SchemaTypes.h
@@ -19,13 +19,13 @@ using FSchemaId = FName;
   * declared in the parent schema.
   * 
   * Example: A lobby contains both lobby data and lobby member data, but during lobby creation the
-  * application only specifies a single schema Id. The schema id is required to have categories for
+  * client only specifies a single schema Id. The schema id is required to have categories for
   * both Lobby and LobbyMember.
   */
 using FSchemaCategoryId = FName;
 /**
-  * The id used to refer to a schema attribute. These ids are declared in the application config
-  * and used in application code to look up an attribute value.
+  * The id used to refer to a schema attribute. These ids are declared in the client config
+  * and used in client code to look up an attribute value.
   */
 using FSchemaAttributeId = FName;
 /**
@@ -307,7 +307,7 @@ struct FSchemaAttributeDescriptor
 	TArray<ESchemaAttributeFlags> Flags;
 	/**
 	  * Whether to group the attribute together with other attributes explicitly.
-	  * Providing a group id allows the application to explicitly serialize attributes to the same
+	  * Providing a group id allows the client to explicitly serialize attributes to the same
 	  * service attribute for increased network efficiency.
 	  *
 	  * Id 0 uses automatic grouping.
@@ -544,17 +544,42 @@ struct FSchemaServiceAttributeData
 	  * Additional behavior required by the attribute. Examples include making an attribute
 	  * public or available to search.
 	  */
-	ESchemaAttributeFlags Flags = ESchemaAttributeFlags::None;
+	ESchemaServiceAttributeFlags Flags = ESchemaServiceAttributeFlags::None;
 	/**
 	  * The attribute value.
 	  */
 	FSchemaVariant Value;
 };
 
+/** Full snapshot of all set service attributes. */
+struct FSchemaServiceSnapshot
+{
+	TMap<FSchemaServiceAttributeId, FSchemaVariant> Attributes;
+};
+
 /**
-  * Changes to be applied to the service data following schema attribute changes.
+  * Changes from the client data to be translated into service changes.
   */
-struct FSchemaServiceChanges
+struct FSchemaClientChanges
+{
+	/**
+	  * Change to schema id.
+	  */
+	TOptional<FSchemaId> SchemaId;
+	/**
+	  * Added or changed attributes and their values.
+	  */
+	TMap<FSchemaAttributeId, FSchemaVariant> UpdatedAttributes;
+	/**
+	  * Attributes to be cleared.
+	  */
+	TSet<FSchemaAttributeId> RemovedAttributes;
+};
+
+/**
+  * Changes to be applied to the service data following client schema attribute changes.
+  */
+struct FSchemaClientServiceChanges
 {
 	/**
 	  * Added or changed attributes and their values.
@@ -563,13 +588,13 @@ struct FSchemaServiceChanges
 	/**
 	  * Removed attribute ids.
 	  */
-	TSet<FSchemaAttributeId> RemovedAttributes;
+	TSet<FSchemaServiceAttributeId> RemovedAttributes;
 };
 
 /**
-  * Changes to be applied to the application data following service attribute changes.
+  * Changes to be applied to the client data following service attribute changes.
   */
-struct FSchemaApplicationChanges
+struct FSchemaServiceClientChanges
 {
 	/**
 	  * Change to schema id.
@@ -589,70 +614,53 @@ struct FSchemaApplicationChanges
 	TSet<FSchemaAttributeId> RemovedAttributes;
 };
 
-#if 0 // todo
-struct FSchemaCategoryInstanceChangeSchema
+struct FSchemaCategoryInstancePrepareClientChanges
 {
-	static constexpr TCHAR Name[] = TEXT("ChangeSchema");
-
 	struct Params
 	{
-		/** Schema Id to switch to. The schema must be configured and be a sibling of the current
-		  * schema.
-		  * A parent schema field definition for SchemaId must exist to enable this feature.
-		  */
-		FSchemaId SchemaId;
+		/** Changes from the client to be applied to the service. */
+		FSchemaClientChanges ClientChanges;
 	};
 
 	struct Result
 	{
 		/** Attribute changes to be notified to the service. */
-		TSharedRef<FSchemaServiceChanges> ServiceChanges;
-
-	public:
-		Result() = delete; // cannot default construct due to TSharedRef
+		FSchemaClientServiceChanges ServiceChanges;
 	};
 };
-#endif
 
-struct FSchemaCategoryInstanceApplyApplicationDelta
+struct FSchemaCategoryInstanceCommitClientChanges
 {
-	static constexpr TCHAR Name[] = TEXT("ApplyApplicationDelta");
-
-	struct Params
-	{
-		/** Attributes notified from the application. */
-		TMap<FSchemaAttributeId, FSchemaVariant> UpdatedAttributes;
-		/** Removed attributes notified from the application. */
-		TSet<FSchemaAttributeId> RemovedAttributes;
-	};
-
 	struct Result
 	{
-		/** Attribute changes to be notified to the service. */
-		FSchemaServiceChanges ServiceChanges;
-
-	public:
-		Result() = delete; // cannot default construct due to TSharedRef
+		/** Attribute changes to be notified to the client. */
+		FSchemaServiceClientChanges ClientChanges;
 	};
 };
 
-struct FSchemaCategoryInstanceApplyServiceSnapshot
+struct FSchemaCategoryInstancePrepareServiceSnapshot
 {
-	static constexpr TCHAR Name[] = TEXT("ApplyServiceSnapshot");
-
 	struct Params
 	{
+		/** Optional schema id for applying a schema change from external state. */
+		TOptional<FSchemaId> SchemaId;
 		/** Full attribute snapshot from the service. */
-		TMap<FSchemaServiceAttributeId, FSchemaVariant> Attributes;
+		FSchemaServiceSnapshot ServiceSnapshot;
 	};
 
 	struct Result
 	{
-		/** Attribute changes to be notified to the application. */
-		FSchemaApplicationChanges ApplicationChanges;
+		/** The new derived schema id when set from the service. */
+		TOptional<FSchemaId> DerivedSchemaId;
+	};
+};
 
-	public:
-		Result() = delete; // cannot default construct due to TSharedRef
+struct FSchemaCategoryInstanceCommitServiceSnapshot
+{
+	struct Result
+	{
+		/** Attribute changes to be notified to the client. */
+		FSchemaServiceClientChanges ClientChanges;
 	};
 };
 
