Refactor of Sessions methods internal architecture to improve stability and reduce duplication

diff --git a/Engine/Plugins/Online/Experimental/OnlineServices/Source/OnlineServicesCommon/Private/Online/SessionsCommon.cpp b/Engine/Plugins/Online/Experimental/OnlineServices/Source/OnlineServicesCommon/Private/Online/SessionsCommon.cpp
--- a/Engine/Plugins/Online/Experimental/OnlineServices/Source/OnlineServicesCommon/Private/Online/SessionsCommon.cpp
+++ b/Engine/Plugins/Online/Experimental/OnlineServices/Source/OnlineServicesCommon/Private/Online/SessionsCommon.cpp
@@ -122,9 +122,25 @@ namespace UE::Online {
 		RegisterCommand(&FSessionsCommon::RejectSessionInvite);
 	}
 
+	// Synchronous methods
+
+#define CHECK_PARAMS_ID_HANDLE(IdHandle, MethodName) \
+	if (!IdHandle.IsValid()) \
+	{ \
+		UE_LOG(LogTemp, Warning, TEXT("[%s] Value of [%s] invalid."), UTF8_TO_TCHAR(__FUNCTION__), UTF8_TO_TCHAR(#IdHandle)); \
+		return TOnlineResult<MethodName>(Errors::InvalidParams()); \
+	} \
+
+#define CHECK_PARAMS_FNAME(Name, MethodName) \
+	if (Name.IsNone()) \
+	{ \
+		UE_LOG(LogTemp, Warning, TEXT("[%s] Value of [%s] invalid."), UTF8_TO_TCHAR(__FUNCTION__), UTF8_TO_TCHAR(#Name)); \
+		return TOnlineResult<MethodName>(Errors::InvalidParams()); \
+	} \
+
 	TOnlineResult<FGetAllSessions> FSessionsCommon::GetAllSessions(FGetAllSessions::Params&& Params) const
 	{
-		// TODO: Params and user login check
+		CHECK_PARAMS_ID_HANDLE(Params.LocalAccountId, FGetAllSessions)
 
 		TArray<TSharedRef<const ISession>> Sessions;
 
@@ -161,16 +177,16 @@ namespace UE::Online {
 		if (GetMutableSessionByIdResult.IsOk())
 		{
 			return TOnlineResult<FGetSessionById>({ GetMutableSessionByIdResult.GetOkValue().Session });
-			}
+		}
 		else
 		{
 			return TOnlineResult<FGetSessionById>(GetMutableSessionByIdResult.GetErrorValue());
-			}
 		}
+	}
 
 	TOnlineResult<FGetPresenceSession> FSessionsCommon::GetPresenceSession(FGetPresenceSession::Params&& Params) const
 	{
-		// TODO: Check the user is valid and logged in
+		CHECK_PARAMS_ID_HANDLE(Params.LocalAccountId, FGetPresenceSession);
 
 		if (const FOnlineSessionId* PresenceSessionId = PresenceSessionsUserMap.Find(Params.LocalAccountId))
 		{
@@ -178,13 +194,15 @@ namespace UE::Online {
 		}
 		else
 		{
+			UE_LOG(LogTemp, Warning, TEXT("[%s] No data found for user [%s]."), UTF8_TO_TCHAR(__FUNCTION__), *ToLogString(Params.LocalAccountId));
 			return TOnlineResult<FGetPresenceSession>({ Errors::InvalidState() });
 		}
 	}
 
 	TOnlineResult<FIsPresenceSession> FSessionsCommon::IsPresenceSession(FIsPresenceSession::Params&& Params) const
 	{
-		// TODO: Check the user is valid and logged in. Check the session id is valid
+		CHECK_PARAMS_ID_HANDLE(Params.LocalAccountId, FIsPresenceSession)
+		CHECK_PARAMS_ID_HANDLE(Params.SessionId, FIsPresenceSession);
 
 		if (const FOnlineSessionId* PresenceSessionId = PresenceSessionsUserMap.Find(Params.LocalAccountId))
 		{
@@ -192,13 +210,15 @@ namespace UE::Online {
 		}
 		else
 		{
+			UE_LOG(LogTemp, Warning, TEXT("[%s] No data found for user [%s]."), UTF8_TO_TCHAR(__FUNCTION__), *ToLogString(Params.LocalAccountId));
 			return TOnlineResult<FIsPresenceSession>({ Errors::InvalidState() });
 		}
 	}
 
 	TOnlineResult<FSetPresenceSession> FSessionsCommon::SetPresenceSession(FSetPresenceSession::Params&& Params)
 	{
-		// TODO: Check the user is valid and logged in. Check the session id is valid
+		CHECK_PARAMS_ID_HANDLE(Params.LocalAccountId, FSetPresenceSession)
+		CHECK_PARAMS_ID_HANDLE(Params.SessionId, FSetPresenceSession);
 
 		FOnlineSessionId& PresenceSessionId = PresenceSessionsUserMap.FindOrAdd(Params.LocalAccountId);
 		PresenceSessionId = Params.SessionId;
@@ -208,835 +228,819 @@ namespace UE::Online {
 
 	TOnlineResult<FClearPresenceSession> FSessionsCommon::ClearPresenceSession(FClearPresenceSession::Params&& Params)
 	{
-		// TODO: Check the user is valid and logged in
+		CHECK_PARAMS_ID_HANDLE(Params.LocalAccountId, FClearPresenceSession);
 
 		PresenceSessionsUserMap.Remove(Params.LocalAccountId);
 
 		return TOnlineResult<FClearPresenceSession>(FClearPresenceSession::Result{ });
 	}
 
-	TOnlineAsyncOpHandle<FCreateSession> FSessionsCommon::CreateSession(FCreateSession::Params&& Params)
+	TOnlineResult<FGetMutableSessionByName> FSessionsCommon::GetMutableSessionByName(FGetMutableSessionByName::Params&& Params) const
 	{
-		TOnlineAsyncOpRef<FCreateSession> Operation = GetOp<FCreateSession>(MoveTemp(Params));
-		Operation->SetError(Errors::NotImplemented());
-		return Operation->GetHandle();
+		CHECK_PARAMS_FNAME(Params.LocalName, FGetMutableSessionByName)
+
+			if (const FOnlineSessionId* SessionId = LocalSessionsByName.Find(Params.LocalName))
+			{
+				check(AllSessionsById.Contains(*SessionId));
+
+				return TOnlineResult<FGetMutableSessionByName>({ AllSessionsById.FindChecked(*SessionId) });
+			}
+			else
+			{
+				UE_LOG(LogTemp, Warning, TEXT("[%s] No session found for name [%s]."), UTF8_TO_TCHAR(__FUNCTION__), *Params.LocalName.ToString());
+				return TOnlineResult<FGetMutableSessionByName>(Errors::NotFound());
+			}
 	}
 
-	TOnlineAsyncOpHandle<FUpdateSessionSettings> FSessionsCommon::UpdateSessionSettings(FUpdateSessionSettings::Params&& Params)
+	TOnlineResult<FGetMutableSessionById> FSessionsCommon::GetMutableSessionById(FGetMutableSessionById::Params&& Params) const
 	{
-		TOnlineAsyncOpRef<FUpdateSessionSettings> Op = GetOp<FUpdateSessionSettings>(MoveTemp(Params));
-		const FUpdateSessionSettings::Params& OpParams = Op->GetParams();
+		CHECK_PARAMS_ID_HANDLE(Params.SessionId, FGetMutableSessionById)
 
-		FOnlineError ParamsCheck = CheckUpdateSessionSettingsParams(OpParams);
-		if (ParamsCheck != Errors::Success())
+		if (const TSharedRef<FSessionCommon>* FoundSession = AllSessionsById.Find(Params.SessionId))
 		{
-			Op->SetError(MoveTemp(ParamsCheck));
-			return Op->GetHandle();
+			return TOnlineResult<FGetMutableSessionById>({ *FoundSession });
+		}
+		else
+		{
+			UE_LOG(LogTemp, Warning, TEXT("[%s] No session found for id [%s]."), UTF8_TO_TCHAR(__FUNCTION__), *ToLogString(Params.SessionId));
+			return TOnlineResult<FGetMutableSessionById>(Errors::NotFound());
 		}
+	}
+
+	TOnlineResult<FGetSessionInviteById> FSessionsCommon::GetSessionInviteById(FGetSessionInviteById::Params&& Params)
+	{
+		CHECK_PARAMS_ID_HANDLE(Params.LocalAccountId, FGetSessionInviteById)
+		CHECK_PARAMS_ID_HANDLE(Params.SessionInviteId, FGetSessionInviteById)
 
-		Op->Then([this](TOnlineAsyncOp<FUpdateSessionSettings>& Op) mutable
+		if (const TMap<FSessionInviteId, TSharedRef<FSessionInvite>>* UserMap = SessionInvitesUserMap.Find(Params.LocalAccountId))
+		{
+			if (const TSharedRef<FSessionInvite>* SessionInvite = UserMap->Find(Params.SessionInviteId))
+			{
+				return TOnlineResult<FGetSessionInviteById>({ *SessionInvite });
+			}
+			else
+			{
+				UE_LOG(LogTemp, Warning, TEXT("[%s] No invite found for invite id [%s]."), UTF8_TO_TCHAR(__FUNCTION__), *ToLogString(Params.SessionInviteId));
+				return TOnlineResult<FGetSessionInviteById>(Errors::NotFound());
+			}
+		}
+		else
 		{
-			const FUpdateSessionSettings::Params& OpParams = Op.GetParams();
+			UE_LOG(LogTemp, Warning, TEXT("[%s] No invite data found for account id [%s]."), UTF8_TO_TCHAR(__FUNCTION__), *ToLogString(Params.LocalAccountId));
+			return TOnlineResult<FGetSessionInviteById>(Errors::InvalidState());
+		}
+	}
 
-			FOnlineError StateCheck = CheckUpdateSessionSettingsState(OpParams);
-			if (StateCheck != Errors::Success())
+	TOnlineResult<FGetAllSessionInvites> FSessionsCommon::GetAllSessionInvites(FGetAllSessionInvites::Params&& Params)
+	{
+		CHECK_PARAMS_ID_HANDLE(Params.LocalAccountId, FGetAllSessionInvites)
+		
+		TArray<TSharedRef<const FSessionInvite>> SessionInvites;
+
+		if (const TMap<FSessionInviteId, TSharedRef<FSessionInvite>>* UserMap = SessionInvitesUserMap.Find(Params.LocalAccountId))
+		{
+			for (const TPair<FSessionInviteId, TSharedRef<FSessionInvite>>& Entry : *UserMap)
 			{
-				Op.SetError(MoveTemp(StateCheck));
-				return;
+				SessionInvites.Add(Entry.Value);
 			}
+		}
+
+		return TOnlineResult<FGetAllSessionInvites>({ MoveTemp(SessionInvites) });
+	}
+
+#undef CHECK_PARAMS_ID_HANDLE
+#undef CHECK_PARAMS_FNAME
+
+	// Asynchronous methods
 
-			UpdateSessionSettingsImpl({ OpParams.LocalAccountId, OpParams.SessionName, OpParams.Mutations })
-			.Next([this, WeakOp = Op.AsWeak()](const TOnlineResult<FUpdateSessionSettingsImpl>& Result)
+#define CHECK_PARAMS_ID_HANDLE(IdHandle) \
+	if (!IdHandle.IsValid()) \
+	{ \
+		UE_LOG(LogTemp, Warning, TEXT("[%s] Value of [%s] invalid."), UTF8_TO_TCHAR(__FUNCTION__), UTF8_TO_TCHAR(#IdHandle)); \
+		return TOptional<FOnlineError>(Errors::InvalidParams()); \
+	} \
+
+#define CHECK_PARAMS_FNAME(Name) \
+	if (Name.IsNone()) \
+	{ \
+		UE_LOG(LogTemp, Warning, TEXT("[%s] Value of [%s] invalid."), UTF8_TO_TCHAR(__FUNCTION__), UTF8_TO_TCHAR(#Name)); \
+		return TOptional<FOnlineError>(Errors::InvalidParams()); \
+	} \
+
+#define CHECK_STATE_NEW_SESSION_NAME(SessionName) \
+	TOnlineResult<FGetSessionByName> GetSessionByNameResult = GetSessionByName({ SessionName }); \
+	if (GetSessionByNameResult.IsOk()) \
+	{ \
+		UE_LOG(LogTemp, Warning, TEXT("[%s] Session with name [%s] already exists."), UTF8_TO_TCHAR(__FUNCTION__), *SessionName.ToString()); \
+		return TOptional<FOnlineError>(Errors::InvalidState()); \
+	} \
+
+#define CHECK_STATE_EXISTING_SESSION_NAME(SessionName) \
+	TOnlineResult<FGetSessionByName> GetSessionByNameResult = GetSessionByName({ SessionName }); \
+	if (GetSessionByNameResult.IsError()) \
+	{ \
+		UE_LOG(LogTemp, Warning, TEXT("[%s] Could not find session with name [%s]."), UTF8_TO_TCHAR(__FUNCTION__), *SessionName.ToString()); \
+		return TOptional<FOnlineError>(Errors::InvalidState()); \
+	} \
+
+	template<typename MethodStruct>
+	TOnlineAsyncOpHandle<MethodStruct> FSessionsCommon::ExecuteAsyncSessionsMethod(typename MethodStruct::Params&& Params, TFuture<TOnlineResult<MethodStruct>>(FSessionsCommon::* ImplFunc)(const typename MethodStruct::Params& Params))
+	{
+		TOnlineAsyncOpRef<MethodStruct> Op = GetOp<MethodStruct>(MoveTemp(Params));
+		const typename MethodStruct::Params& OpParams = Op->GetParams();
+
+		if (TOptional<FOnlineError> ParamsCheck = CheckParams(OpParams))
+		{
+			Op->SetError(MoveTemp(ParamsCheck.GetValue()));
+			return Op->GetHandle();
+		}
+
+		Op->Then([this, WeakThis = AsWeak(), ImplFunc](TOnlineAsyncOp<MethodStruct>& Op)
 			{
-				if (TOnlineAsyncOpPtr<FUpdateSessionSettings> StrongOp = WeakOp.Pin())
+				if (TSharedPtr<ISessions> StrongThis = WeakThis.Pin())
 				{
-					if (Result.IsOk())
-					{
-						StrongOp->SetResult({ });
-					}
-					else
+					const typename MethodStruct::Params& OpParams = Op.GetParams();
+
+					if (TOptional<FOnlineError> StateCheck = CheckState(Op.GetParams()))
 					{
-						FOnlineError ErrorValue = Result.GetErrorValue();
-						StrongOp->SetError(MoveTemp(ErrorValue));
+						Op.SetError(MoveTemp(StateCheck.GetValue()));
+						return;
 					}
+
+					(this->*ImplFunc)(OpParams)
+						.Next([WeakOp = Op.AsWeak()](const TOnlineResult<MethodStruct>& Result)
+							{
+								if (TOnlineAsyncOpPtr<MethodStruct> StrongOp = WeakOp.Pin())
+								{
+									if (Result.IsOk())
+									{
+										StrongOp->SetResult(typename MethodStruct::Result(Result.GetOkValue()));
+									}
+									else
+									{
+										StrongOp->SetError(FOnlineError(Result.GetErrorValue()));
+									}
+								}
+							});
 				}
-			});			
-		})
-		.Enqueue(GetSerialQueue());
+			})
+			.Enqueue(GetSerialQueue());
 
 		return Op->GetHandle();
 	}
 
-	TOnlineAsyncOpHandle<FLeaveSession> FSessionsCommon::LeaveSession(FLeaveSession::Params&& Params)
-	{
-		TOnlineAsyncOpRef<FLeaveSession> Operation = GetOp<FLeaveSession>(MoveTemp(Params));
-		Operation->SetError(Errors::NotImplemented());
-		return Operation->GetHandle();
-	}
-
-	TOnlineAsyncOpHandle<FFindSessions> FSessionsCommon::FindSessions(FFindSessions::Params&& Params)
-	{
-		TOnlineAsyncOpRef<FFindSessions> Operation = GetOp<FFindSessions>(MoveTemp(Params));
-		Operation->SetError(Errors::NotImplemented());
-		return Operation->GetHandle();
-	}
+	// CreateSession
 
-	TOnlineAsyncOpHandle<FStartMatchmaking> FSessionsCommon::StartMatchmaking(FStartMatchmaking::Params&& Params)
+	TOnlineAsyncOpHandle<FCreateSession> FSessionsCommon::CreateSession(FCreateSession::Params&& Params)
 	{
-		TOnlineAsyncOpRef<FStartMatchmaking> Operation = GetOp<FStartMatchmaking>(MoveTemp(Params));
-		Operation->SetError(Errors::NotImplemented());
-		return Operation->GetHandle();
+		return ExecuteAsyncSessionsMethod<FCreateSession>(MoveTemp(Params), &FSessionsCommon::CreateSessionImpl);
 	}
 
-	TOnlineAsyncOpHandle<FJoinSession> FSessionsCommon::JoinSession(FJoinSession::Params&& Params)
+	TFuture<TOnlineResult<FCreateSession>> FSessionsCommon::CreateSessionImpl(const FCreateSession::Params& Params)
 	{
-		TOnlineAsyncOpRef<FJoinSession> Operation = GetOp<FJoinSession>(MoveTemp(Params));
-		Operation->SetError(Errors::NotImplemented());
-		return Operation->GetHandle();
+		return MakeFulfilledPromise<TOnlineResult<FCreateSession>>(Errors::NotImplemented()).GetFuture();
 	}
 
-	TOnlineAsyncOpHandle<FAddSessionMember> FSessionsCommon::AddSessionMember(FAddSessionMember::Params&& Params)
+	TOptional<FOnlineError> FSessionsCommon::CheckParams(const FCreateSession::Params& Params) const
 	{
-		TOnlineAsyncOpRef<FAddSessionMember> Op = GetOp<FAddSessionMember>(MoveTemp(Params));
+		CHECK_PARAMS_ID_HANDLE(Params.LocalAccountId)
+		CHECK_PARAMS_FNAME(Params.SessionName)
+		CHECK_PARAMS_FNAME(Params.SessionSettings.SchemaName)
 
-		Op->Then([this](TOnlineAsyncOp<FAddSessionMember>& Op) mutable
+		if (Params.SessionSettings.NumMaxConnections == 0)
 		{
-			const FAddSessionMember::Params& OpParams = Op.GetParams();
-
-			FOnlineError StateCheck = CheckAddSessionMemberState(OpParams);
-			if (StateCheck != Errors::Success())
-			{
-				Op.SetError(MoveTemp(StateCheck));
-				return;
-			}
+			UE_LOG(LogTemp, Warning, TEXT("[%s] NumMaxConnections invalid, can't be 0."), UTF8_TO_TCHAR(__FUNCTION__));
+			return Errors::InvalidParams();
+		}
 
-			TOnlineResult<FAddSessionMember> Result = AddSessionMemberImpl(OpParams);
-			if (Result.IsOk())
-			{
-				Op.SetResult(MoveTemp(Result.GetOkValue()));
-			}
-			else
-			{
-				Op.SetError(MoveTemp(Result.GetErrorValue()));
-			}
-		})
-		.Enqueue(GetSerialQueue());
+		TArray<FSchemaAttributeId> AttributeIds;
+		Params.SessionSettings.CustomSettings.GenerateKeyArray(AttributeIds);
+		for (const FSchemaAttributeId& AttributeId : AttributeIds)
+		{
+			CHECK_PARAMS_FNAME(AttributeId)
+		}
 
-		return Op->GetHandle();
+		return TOptional<FOnlineError>();
 	}
 
-	TOnlineAsyncOpHandle<FRemoveSessionMember> FSessionsCommon::RemoveSessionMember(FRemoveSessionMember::Params&& Params)
+	TOptional<FOnlineError> FSessionsCommon::CheckState(const FCreateSession::Params& Params) const
 	{
-		TOnlineAsyncOpRef<FRemoveSessionMember> Op = GetOp<FRemoveSessionMember>(MoveTemp(Params));
+		CHECK_STATE_NEW_SESSION_NAME(Params.SessionName)
 
-		Op->Then([this](TOnlineAsyncOp<FRemoveSessionMember>& Op) mutable
+		if (Params.bPresenceEnabled)
 		{
-			const FRemoveSessionMember::Params& OpParams = Op.GetParams();
-
-			FOnlineError StateCheck = CheckRemoveSessionMemberState(OpParams);
-			if (StateCheck != Errors::Success())
+			for (const TPair<FName, FOnlineSessionId>& Entry : LocalSessionsByName)
 			{
-				Op.SetError(MoveTemp(StateCheck));
-				return;
+				TOnlineResult<FGetPresenceSession> GetPresenceSessionResult = GetPresenceSession({ Params.LocalAccountId });
+				if (GetPresenceSessionResult.IsOk())
+				{
+					UE_LOG(LogTemp, Warning, TEXT("[%s] Could not create session with bPresenceEnabled set to true when another already exists [%s]."), UTF8_TO_TCHAR(__FUNCTION__), *Entry.Key.ToString());
+					return Errors::InvalidState();
+				}
 			}
+		}
 
-			TOnlineResult<FRemoveSessionMember> Result = RemoveSessionMemberImpl(OpParams);
-			if (Result.IsOk())
-			{
-				Op.SetResult(MoveTemp(Result.GetOkValue()));
-			}
-			else
-			{
-				Op.SetError(MoveTemp(Result.GetErrorValue()));
-			}
-		})
-		.Enqueue(GetSerialQueue());
+		return TOptional<FOnlineError>();
+	}
 
-		return Op->GetHandle();
+	// UpdateSessionSettings
+
+	TOnlineAsyncOpHandle<FUpdateSessionSettings> FSessionsCommon::UpdateSessionSettings(FUpdateSessionSettings::Params&& Params)
+	{
+		return ExecuteAsyncSessionsMethod<FUpdateSessionSettings>(MoveTemp(Params), &FSessionsCommon::UpdateSessionSettingsImpl);
 	}
 
-	TOnlineAsyncOpHandle<FSendSessionInvite> FSessionsCommon::SendSessionInvite(FSendSessionInvite::Params&& Params)
+	TFuture<TOnlineResult<FUpdateSessionSettings>> FSessionsCommon::UpdateSessionSettingsImpl(const FUpdateSessionSettings::Params& Params)
 	{
-		TOnlineAsyncOpRef<FSendSessionInvite> Operation = GetOp<FSendSessionInvite>(MoveTemp(Params));
-		Operation->SetError(Errors::NotImplemented());
-		return Operation->GetHandle();
+		return MakeFulfilledPromise<TOnlineResult<FUpdateSessionSettings>>(Errors::NotImplemented()).GetFuture();
 	}
 
-	TOnlineResult<FGetSessionInviteById> FSessionsCommon::GetSessionInviteById(FGetSessionInviteById::Params&& Params)
+	TOptional<FOnlineError> FSessionsCommon::CheckParams(const FUpdateSessionSettings::Params& Params) const
 	{
-		if (const TMap<FSessionInviteId, TSharedRef<FSessionInvite>>* UserMap = SessionInvitesUserMap.Find(Params.LocalAccountId))
+		CHECK_PARAMS_ID_HANDLE(Params.LocalAccountId)
+		CHECK_PARAMS_FNAME(Params.SessionName)
+
+		if (Params.Mutations.SchemaName.IsSet())
 		{
-			if (const TSharedRef<FSessionInvite>* SessionInvite = UserMap->Find(Params.SessionInviteId))
-			{
-				return TOnlineResult<FGetSessionInviteById>({ *SessionInvite });
-			}
-			else
+			CHECK_PARAMS_FNAME(Params.Mutations.SchemaName.GetValue())
+		}
+
+		if (Params.Mutations.NumMaxConnections.IsSet())
+		{
+			if (Params.Mutations.NumMaxConnections.GetValue() == 0)
 			{
-				return TOnlineResult<FGetSessionInviteById>(Errors::NotFound());
+				UE_LOG(LogTemp, Warning, TEXT("[%s] NumMaxConnections invalid, can't be 0."), UTF8_TO_TCHAR(__FUNCTION__));
+				return Errors::InvalidParams();
 			}
 		}
-		else
-		{
-			return TOnlineResult<FGetSessionInviteById>(Errors::InvalidState());
-		}		
-	}
 
-	TOnlineResult<FGetAllSessionInvites> FSessionsCommon::GetAllSessionInvites(FGetAllSessionInvites::Params&& Params)
-	{
-		TArray<TSharedRef<const FSessionInvite>> SessionInvites;
+		TArray<FSchemaAttributeId> UpdatedAttributeIds;
+		Params.Mutations.UpdatedCustomSettings.GenerateKeyArray(UpdatedAttributeIds);
+		for (const FSchemaAttributeId& UpdatedAttributeId : UpdatedAttributeIds)
+		{
+			CHECK_PARAMS_FNAME(UpdatedAttributeId)
+		}
 
-		if (const TMap<FSessionInviteId, TSharedRef<FSessionInvite>>* UserMap = SessionInvitesUserMap.Find(Params.LocalAccountId))
+		for (const FSchemaAttributeId& RemovedAttributeId : Params.Mutations.RemovedCustomSettings)
 		{
-			for (const TPair<FSessionInviteId, TSharedRef<FSessionInvite>>& Entry : *UserMap)
-			{
-				SessionInvites.Add(Entry.Value);
-			}
+			CHECK_PARAMS_FNAME(RemovedAttributeId)
 		}
 
-		return TOnlineResult<FGetAllSessionInvites>({ MoveTemp(SessionInvites) });
+		return TOptional<FOnlineError>();
 	}
 
-	TOnlineAsyncOpHandle<FRejectSessionInvite> FSessionsCommon::RejectSessionInvite(FRejectSessionInvite::Params&& Params)
+	TOptional<FOnlineError> FSessionsCommon::CheckState(const FUpdateSessionSettings::Params& Params) const
 	{
-		TOnlineAsyncOpRef<FRejectSessionInvite> Operation = GetOp<FRejectSessionInvite>(MoveTemp(Params));
-		Operation->SetError(Errors::NotImplemented());
-		return Operation->GetHandle();
+		CHECK_STATE_EXISTING_SESSION_NAME(Params.SessionName)
+
+		return TOptional<FOnlineError>();
 	}
 
-	TOnlineEvent<void(const FSessionJoined&)> FSessionsCommon::OnSessionJoined()
+	// LeaveSession
+
+	TOnlineAsyncOpHandle<FLeaveSession> FSessionsCommon::LeaveSession(FLeaveSession::Params&& Params)
 	{
-		return SessionEvents.OnSessionJoined;
+		return ExecuteAsyncSessionsMethod<FLeaveSession>(MoveTemp(Params), &FSessionsCommon::LeaveSessionImpl);
 	}
 
-	TOnlineEvent<void(const FSessionLeft&)> FSessionsCommon::OnSessionLeft()
+	TFuture<TOnlineResult<FLeaveSession>> FSessionsCommon::LeaveSessionImpl(const FLeaveSession::Params& Params)
 	{
-		return SessionEvents.OnSessionLeft;
+		return MakeFulfilledPromise<TOnlineResult<FLeaveSession>>(Errors::NotImplemented()).GetFuture();
 	}
 
-	TOnlineEvent<void(const FSessionUpdated&)> FSessionsCommon::OnSessionUpdated()
+	TOptional<FOnlineError> FSessionsCommon::CheckParams(const FLeaveSession::Params& Params) const
 	{
-		return SessionEvents.OnSessionUpdated;
+		CHECK_PARAMS_ID_HANDLE(Params.LocalAccountId)
+		CHECK_PARAMS_FNAME(Params.SessionName)
+
+		return TOptional<FOnlineError>();
 	}
 
-	TOnlineEvent<void(const FSessionInviteReceived&)> FSessionsCommon::OnSessionInviteReceived()
+	TOptional<FOnlineError> FSessionsCommon::CheckState(const FLeaveSession::Params& Params) const
 	{
-		return SessionEvents.OnSessionInviteReceived;
+		CHECK_STATE_EXISTING_SESSION_NAME(Params.SessionName)
+
+		return TOptional<FOnlineError>();
 	}
 
-	TOnlineEvent<void(const FUISessionJoinRequested&)> FSessionsCommon::OnUISessionJoinRequested()
+	// FindSessions
+
+	TOnlineAsyncOpHandle<FFindSessions> FSessionsCommon::FindSessions(FFindSessions::Params&& Params)
 	{
-		return SessionEvents.OnUISessionJoinRequested;
+		return ExecuteAsyncSessionsMethod<FFindSessions>(MoveTemp(Params), &FSessionsCommon::FindSessionsImpl);
 	}
 
-	TOnlineResult<FGetMutableSessionByName> FSessionsCommon::GetMutableSessionByName(FGetMutableSessionByName::Params&& Params) const
+	TFuture<TOnlineResult<FFindSessions>> FSessionsCommon::FindSessionsImpl(const FFindSessions::Params& Params)
 	{
-		// TODO: Check that name is not empty
+		return MakeFulfilledPromise<TOnlineResult<FFindSessions>>(Errors::NotImplemented()).GetFuture();
+	}
 
-		if (const FOnlineSessionId* SessionId = LocalSessionsByName.Find(Params.LocalName))
-		{
-			check(AllSessionsById.Contains(*SessionId));
+	TOptional<FOnlineError> FSessionsCommon::CheckParams(const FFindSessions::Params& Params) const
+	{
+		CHECK_PARAMS_ID_HANDLE(Params.LocalAccountId)
 
-			return TOnlineResult<FGetMutableSessionByName>({ AllSessionsById.FindChecked(*SessionId) });
-		}
-		else
+		if (Params.MaxResults == 0)
 		{
-			return TOnlineResult<FGetMutableSessionByName>(Errors::NotFound());
+			UE_LOG(LogTemp, Warning, TEXT("[%s] MaxResults can't be 0"), UTF8_TO_TCHAR(__FUNCTION__));
+			return TOptional<FOnlineError>(Errors::InvalidParams());
 		}
-	}
 
-	TOnlineResult<FGetMutableSessionById> FSessionsCommon::GetMutableSessionById(FGetMutableSessionById::Params&& Params) const
-	{
-		// TODO: Check that session id is valid
+		for (const FFindSessionsSearchFilter& Filter : Params.Filters)
+		{
+			CHECK_PARAMS_FNAME(Filter.Key)
+		}
 
-		if (const TSharedRef<FSessionCommon>* FoundSession = AllSessionsById.Find(Params.SessionId))
+		if (Params.SessionId.IsSet())
 		{
-			return TOnlineResult<FGetMutableSessionById>({ *FoundSession });
+			CHECK_PARAMS_ID_HANDLE(Params.SessionId.GetValue())
 		}
-		else
+
+		if (Params.TargetUser.IsSet())
 		{
-			return TOnlineResult<FGetMutableSessionById>(Errors::NotFound());
+			CHECK_PARAMS_ID_HANDLE(Params.TargetUser.GetValue())
 		}
+		
+		return TOptional<FOnlineError>();
 	}
 
-	void FSessionsCommon::AddSessionInvite(const TSharedRef<FSessionInvite> SessionInvite, const TSharedRef<FSessionCommon> Session, const FAccountId& LocalAccountId)
+	TOptional<FOnlineError> FSessionsCommon::CheckState(const FFindSessions::Params& Params) const
 	{
-		AllSessionsById.Emplace(Session->GetSessionId(), Session);
+		if (CurrentSessionSearchPromisesUserMap.Contains(Params.LocalAccountId))
+		{
+			UE_LOG(LogTemp, Warning, TEXT("[%s] Search already in progress"), UTF8_TO_TCHAR(__FUNCTION__));
+			return Errors::AlreadyPending();
+		}
 
-		TMap<FSessionInviteId, TSharedRef<FSessionInvite>>& SessionInvitesMap = SessionInvitesUserMap.FindOrAdd(LocalAccountId);
-		SessionInvitesMap.Emplace(SessionInvite->InviteId, SessionInvite);
+		return TOptional<FOnlineError>();
 	}
 
-	void FSessionsCommon::AddSearchResult(const TSharedRef<FSessionCommon> Session, const FAccountId& LocalAccountId)
-	{
-		AllSessionsById.Emplace(Session->GetSessionId(), Session);
+	// StartMatchmaking
 
-		TArray<FOnlineSessionId>& SearchResults = SearchResultsUserMap.FindOrAdd(LocalAccountId);
-		SearchResults.Add(Session->GetSessionId());
+	TOnlineAsyncOpHandle<FStartMatchmaking> FSessionsCommon::StartMatchmaking(FStartMatchmaking::Params&& Params)
+	{
+		return ExecuteAsyncSessionsMethod<FStartMatchmaking>(MoveTemp(Params), &FSessionsCommon::StartMatchmakingImpl);
 	}
 
-	void FSessionsCommon::AddSessionWithReferences(const TSharedRef<FSessionCommon> Session, const FName& SessionName, const FAccountId& LocalAccountId, bool bIsPresenceSession)
+	TFuture<TOnlineResult<FStartMatchmaking>> FSessionsCommon::StartMatchmakingImpl(const FStartMatchmaking::Params& Params)
 	{
-		AllSessionsById.Emplace(Session->GetSessionId(), Session);
-
-		AddSessionReferences(Session->GetSessionId(), SessionName, LocalAccountId, bIsPresenceSession);
+		return MakeFulfilledPromise<TOnlineResult<FStartMatchmaking>>(Errors::NotImplemented()).GetFuture();
 	}
 
-	void FSessionsCommon::AddSessionReferences(const FOnlineSessionId SessionId, const FName& SessionName, const FAccountId& LocalAccountId, bool bIsPresenceSession)
+	TOptional<FOnlineError> FSessionsCommon::CheckParams(const FStartMatchmaking::Params& Params) const
 	{
-		LocalSessionsByName.Emplace(SessionName, SessionId);
-
-		NamedSessionUserMap.FindOrAdd(LocalAccountId).AddUnique(SessionName);
-
-		if (bIsPresenceSession)
+		if (TOptional<FOnlineError> Result = CheckParams(Params.SessionCreationParameters))
 		{
-			SetPresenceSession({ LocalAccountId, SessionId });
+			return Result;
 		}
-	}
 
-	void FSessionsCommon::ClearSessionInvitesForSession(const FAccountId& LocalAccountId, const FOnlineSessionId SessionId)
-	{
-		if (TMap<FSessionInviteId, TSharedRef<FSessionInvite>>* UserMap = SessionInvitesUserMap.Find(LocalAccountId))
+		for (const FFindSessionsSearchFilter& Filter : Params.SessionSearchFilters)
 		{
-			TArray<FSessionInviteId> InviteIdsToRemove;
-			for (const TPair<FSessionInviteId, TSharedRef<FSessionInvite>>& Entry : *UserMap)
-			{
-				if (Entry.Value->SessionId == SessionId)
-				{
-					InviteIdsToRemove.Add(Entry.Key);
-				}
-			}
-
-			for (const FSessionInviteId& InviteId : InviteIdsToRemove)
-			{
-				UserMap->Remove(InviteId);
-			}
+			CHECK_PARAMS_FNAME(Filter.Key)
 		}
+
+		return TOptional<FOnlineError>();
 	}
 
-	void FSessionsCommon::ClearSessionReferences(const FOnlineSessionId SessionId, const FName& SessionName, const FAccountId& LocalAccountId)
+	TOptional<FOnlineError> FSessionsCommon::CheckState(const FStartMatchmaking::Params& Params) const
 	{
-		NamedSessionUserMap.FindChecked(LocalAccountId).Remove(SessionName);
-
-
-		TOnlineResult<FIsPresenceSession> IsPresenceSessionResult = IsPresenceSession({ LocalAccountId, SessionId });
-		if (IsPresenceSessionResult.IsOk())
+		if (TOptional<FOnlineError> Result = CheckState(Params.SessionCreationParameters))
 		{
-			if (IsPresenceSessionResult.GetOkValue().bIsPresenceSession)
-			{
-				ClearPresenceSession({ LocalAccountId });
-			}
+			return Result;
 		}
 
-		ClearSessionByName(SessionName);
-		ClearSessionById(SessionId);
+		return TOptional<FOnlineError>();
 	}
 
-	void FSessionsCommon::ClearSessionByName(const FName& SessionName)
+	// JoinSession
+
+	TOnlineAsyncOpHandle<FJoinSession> FSessionsCommon::JoinSession(FJoinSession::Params&& Params)
 	{
-		for (const TPair<FAccountId, TArray<FName>>& Entry : NamedSessionUserMap)
-		{
-			if (Entry.Value.Contains(SessionName))
-			{
-				return;
-			}
-		}
-
-		// If no references were found, we'll remove the named session entry
-		LocalSessionsByName.Remove(SessionName);
+		return ExecuteAsyncSessionsMethod<FJoinSession>(MoveTemp(Params), &FSessionsCommon::JoinSessionImpl);
 	}
 
-	void FSessionsCommon::ClearSessionById(const FOnlineSessionId& SessionId)
+	TFuture<TOnlineResult<FJoinSession>> FSessionsCommon::JoinSessionImpl(const FJoinSession::Params& Params)
 	{
-		// PresenceSessionsUserMap is not evaluated, since any session there would also be in LocalSessionsByName
-		for (const TPair<FName, FOnlineSessionId>& Entry : LocalSessionsByName)
-		{
-			if (Entry.Value == SessionId)
-			{
-				return;
-			}
-		}
-
-		for (const TPair<FAccountId, TMap<FSessionInviteId, TSharedRef<FSessionInvite>>>& Entry : SessionInvitesUserMap)
-		{
-			for (const TPair<FSessionInviteId, TSharedRef<FSessionInvite>>& InviteMap : Entry.Value)
-			{
-				if (InviteMap.Value->SessionId == SessionId)
-				{
-					return;
-				}
-			}			
-		}
+		return MakeFulfilledPromise<TOnlineResult<FJoinSession>>(Errors::NotImplemented()).GetFuture();
+	}
 
-		for (const TPair<FAccountId, TArray<FOnlineSessionId>>& Entry : SearchResultsUserMap)
-		{
-			if (Entry.Value.Contains(SessionId))
-			{
-				return;
-			}
-		}
+	TOptional<FOnlineError> FSessionsCommon::CheckParams(const FJoinSession::Params& Params) const
+	{
+		CHECK_PARAMS_ID_HANDLE(Params.LocalAccountId)
+		CHECK_PARAMS_ID_HANDLE(Params.SessionId)
+		CHECK_PARAMS_FNAME(Params.SessionName)
 
-		// If no references were found, we'll remove the session entry
-		AllSessionsById.Remove(SessionId);
+		return TOptional<FOnlineError>();
 	}
 
-#define COPY_TOPTIONAL_VALUE_IF_SET(Value) \
-	if (UpdatedValues.Value.IsSet()) \
-	{ \
-		SessionSettingsChanges.Value = UpdatedValues.Value.GetValue(); \
-	} \
-
-	FSessionUpdate FSessionsCommon::BuildSessionUpdate(const TSharedRef<FSessionCommon>& Session, const FSessionSettingsUpdate& UpdatedValues) const
+	TOptional<FOnlineError> FSessionsCommon::CheckState(const FJoinSession::Params& Params) const
 	{
-		FSessionUpdate Result;
+		CHECK_STATE_NEW_SESSION_NAME(Params.SessionName)
 
-		FSessionSettingsChanges SessionSettingsChanges;
+		TOnlineResult<FGetSessionById> GetSessionByIdResult = GetSessionById({ Params.SessionId });
+		if (GetSessionByIdResult.IsError())
+		{
+			UE_LOG(LogTemp, Warning, TEXT("[%s] Session [%s] not found. Please call FindSessions to get an updated list of available sessions "), UTF8_TO_TCHAR(__FUNCTION__), *ToLogString(Params.SessionId));
+			return TOptional<FOnlineError>(GetSessionByIdResult.GetErrorValue());
+		}
 
-		COPY_TOPTIONAL_VALUE_IF_SET(SchemaName) // TODO: We may need some additional logic for schema changes
-		COPY_TOPTIONAL_VALUE_IF_SET(NumMaxConnections)
-		COPY_TOPTIONAL_VALUE_IF_SET(JoinPolicy)
-		COPY_TOPTIONAL_VALUE_IF_SET(bAllowNewMembers)
+		TSharedRef<const ISession> FoundSession = GetSessionByIdResult.GetOkValue().Session;
 
-		SessionSettingsChanges.RemovedCustomSettings.Append(UpdatedValues.RemovedCustomSettings);
+		const FSessionSettings& SessionSettings = FoundSession->GetSessionSettings();
 
-		for (const TPair<FSchemaAttributeId, FCustomSessionSetting>& Entry : UpdatedValues.UpdatedCustomSettings)
+		if (FoundSession->GetSessionMembers().Contains(Params.LocalAccountId))
 		{
-			if (const FCustomSessionSetting* CustomSetting = Session->GetSessionSettings().CustomSettings.Find(Entry.Key))
-			{
-				FCustomSessionSettingUpdate SettingUpdate = { *CustomSetting, Entry.Value };
+			UE_LOG(LogTemp, Warning, TEXT("[%s] User [%s] already in session [%s]"), UTF8_TO_TCHAR(__FUNCTION__), *ToLogString(Params.LocalAccountId), *ToLogString(Params.SessionId));
+			return TOptional<FOnlineError>(Errors::AccessDenied());
+		}
 
-				SessionSettingsChanges.ChangedCustomSettings.Emplace(Entry.Key, SettingUpdate);
-			}
-			else
+		if (!FoundSession->IsJoinable())
+		{
+			UE_LOG(LogTemp, Warning, TEXT("[%s] Session [%s] not joinable "), UTF8_TO_TCHAR(__FUNCTION__), *ToLogString(Params.SessionId));
+			return TOptional<FOnlineError>(Errors::AccessDenied());
+		}
+
+		if (Params.bPresenceEnabled)
+		{
+			for (const TPair<FName, FOnlineSessionId>& Entry : LocalSessionsByName)
 			{
-				SessionSettingsChanges.AddedCustomSettings.Add(Entry);
+				TOnlineResult<FGetPresenceSession> GetPresenceSessionResult = GetPresenceSession({ Params.LocalAccountId });
+				if (GetPresenceSessionResult.IsOk())
+				{
+					UE_LOG(LogTemp, Warning, TEXT("[%s] Could not join session with bPresenceEnabled set to true when another presence session already exists [%s]."), UTF8_TO_TCHAR(__FUNCTION__), *Entry.Key.ToString());
+					return TOptional<FOnlineError>(Errors::InvalidState());
+				}
 			}
 		}
 
-		Result.SessionSettingsChanges = SessionSettingsChanges;
-
-		return Result;
+		return TOptional<FOnlineError>();
 	}
 
-#undef COPY_TOPTIONAL_VALUE_IF_SET
+	// AddSessionMember
 
-	TFuture<TOnlineResult<FUpdateSessionSettingsImpl>> FSessionsCommon::UpdateSessionSettingsImpl(FUpdateSessionSettingsImpl::Params&& Params)
+	TOnlineAsyncOpHandle<FAddSessionMember> FSessionsCommon::AddSessionMember(FAddSessionMember::Params&& Params)
 	{
-		return MakeFulfilledPromise<TOnlineResult<FUpdateSessionSettingsImpl>>(Errors::NotImplemented()).GetFuture();
+		return ExecuteAsyncSessionsMethod<FAddSessionMember>(MoveTemp(Params), &FSessionsCommon::AddSessionMemberImpl);
 	}
 
-	TOnlineResult<FAddSessionMember> FSessionsCommon::AddSessionMemberImpl(const FAddSessionMember::Params& Params)
+	TFuture<TOnlineResult<FAddSessionMember>> FSessionsCommon::AddSessionMemberImpl(const FAddSessionMember::Params& Params)
+	{
+		TPromise<TOnlineResult<FAddSessionMember>> Promise;
+		TFuture<TOnlineResult<FAddSessionMember>> Future = Promise.GetFuture();
+
+		Promise.EmplaceValue(AddSessionMemberInternal(Params));		
+
+		return Future;
+	}
+
+	TOnlineResult<FAddSessionMember> FSessionsCommon::AddSessionMemberInternal(const FAddSessionMember::Params& Params)
 	{
 		TOnlineResult<FGetMutableSessionByName> GetMutableSessionByNameResult = GetMutableSessionByName({ Params.SessionName });
 		if (GetMutableSessionByNameResult.IsError())
 		{
 			return TOnlineResult<FAddSessionMember>(GetMutableSessionByNameResult.GetErrorValue());
 		}
+
 		TSharedRef<FSessionCommon> FoundSession = GetMutableSessionByNameResult.GetOkValue().Session;
 
 		FSessionSettings& SessionSettings = FoundSession->SessionSettings;
 
-		if (FoundSession->GetNumOpenConnections() <= 0)
-		{
-			return TOnlineResult<FAddSessionMember>(Errors::InvalidState());
-		}
-
 		FoundSession->SessionMembers.Emplace(Params.LocalAccountId);
 
 		return TOnlineResult<FAddSessionMember>();
 	}
 
-	TOnlineResult<FRemoveSessionMember> FSessionsCommon::RemoveSessionMemberImpl(const FRemoveSessionMember::Params& Params)
+	TOptional<FOnlineError> FSessionsCommon::CheckParams(const FAddSessionMember::Params& Params) const
 	{
-		TOnlineResult<FGetMutableSessionByName> GetMutableSessionByNameResult = GetMutableSessionByName({ Params.SessionName });
-		if (GetMutableSessionByNameResult.IsError())
-		{
-			return TOnlineResult<FRemoveSessionMember>(GetMutableSessionByNameResult.GetErrorValue());
-		}
-		TSharedRef<FSessionCommon> FoundSession = GetMutableSessionByNameResult.GetOkValue().Session;
+		CHECK_PARAMS_ID_HANDLE(Params.LocalAccountId)
+		CHECK_PARAMS_FNAME(Params.SessionName)
 
-		FSessionSettings& SessionSettings = FoundSession->SessionSettings;
-		
-		if (FoundSession->GetNumOpenConnections() == SessionSettings.NumMaxConnections)
-		{
-			return TOnlineResult<FRemoveSessionMember>(Errors::InvalidState());
-		}
-
-		FoundSession->SessionMembers.Remove(Params.LocalAccountId);
-
-		return TOnlineResult<FRemoveSessionMember>();
+		return TOptional<FOnlineError>();
 	}
 
-	FOnlineError FSessionsCommon::CheckCreateSessionParams(const FCreateSession::Params& Params)
+	TOptional<FOnlineError> FSessionsCommon::CheckState(const FAddSessionMember::Params& Params) const
 	{
-		if (Params.SessionName.IsNone())
-		{
-			UE_LOG(LogTemp, Warning, TEXT("[FSessionsCommon::CheckCreateSessionParams] Could not create session with no valid SessionName set"));
+		CHECK_STATE_EXISTING_SESSION_NAME(Params.SessionName)
 
-			return Errors::InvalidParams();
-		}
+		return TOptional<FOnlineError>();
+	}
 
-		if (!Params.LocalAccountId.IsValid())
-		{
-			UE_LOG(LogTemp, Warning, TEXT("[FSessionsCommon::CheckCreateSessionParams] Could not create session with name [%s]. LocalAccountId [%s] not valid"), *Params.SessionName.ToString(), *ToLogString(Params.LocalAccountId));
+	// RemoveSessionMember
 
-			return Errors::InvalidParams();
-		}
+	TOnlineAsyncOpHandle<FRemoveSessionMember> FSessionsCommon::RemoveSessionMember(FRemoveSessionMember::Params&& Params)
+	{
+		return ExecuteAsyncSessionsMethod<FRemoveSessionMember>(MoveTemp(Params), &FSessionsCommon::RemoveSessionMemberImpl);
+	}
 
-		if (Params.SessionSettings.NumMaxConnections == 0)
-		{
-			UE_LOG(LogTemp, Warning, TEXT("[FSessionsCommon::CheckCreateSessionParams] Could not create session with name [%s] with no valid NumMaxConnections [%d]"), *Params.SessionName.ToString(), Params.SessionSettings.NumMaxConnections);
+	TFuture<TOnlineResult<FRemoveSessionMember>> FSessionsCommon::RemoveSessionMemberImpl(const FRemoveSessionMember::Params& Params)
+	{
+		TPromise<TOnlineResult<FRemoveSessionMember>> Promise;
+		TFuture<TOnlineResult<FRemoveSessionMember>> Future = Promise.GetFuture();
 
-			return Errors::InvalidParams();
-		}
+		Promise.EmplaceValue(RemoveSessionMemberInternal(Params));
 
-		return Errors::Success();
+		return Future;
 	}
 
-	FOnlineError FSessionsCommon::CheckCreateSessionState(const FCreateSession::Params& Params)
+	TOnlineResult<FRemoveSessionMember> FSessionsCommon::RemoveSessionMemberInternal(const FRemoveSessionMember::Params& Params)
 	{
-		TOnlineResult<FGetSessionByName> GetSessionByNameResult = GetSessionByName({ Params.SessionName });
-		if (GetSessionByNameResult.IsOk())
+		TOnlineResult<FGetMutableSessionByName> GetMutableSessionByNameResult = GetMutableSessionByName({ Params.SessionName });
+		if (GetMutableSessionByNameResult.IsError())
 		{
-			UE_LOG(LogTemp, Warning, TEXT("[FSessionsCommon::CheckCreateSessionState] Could not create session with name [%s]. A session with that name already exists"), *Params.SessionName.ToString());
-
-			return Errors::InvalidState();
+			return TOnlineResult<FRemoveSessionMember>(GetMutableSessionByNameResult.GetErrorValue());
 		}
 
-		if (Params.bPresenceEnabled)
-		{
-			for (const TPair<FName, FOnlineSessionId>& Entry : LocalSessionsByName)
-			{
-				TOnlineResult<FGetPresenceSession> GetPresenceSessionResult = GetPresenceSession({ Params.LocalAccountId });
-				if (GetPresenceSessionResult.IsOk())
-				{
-					UE_LOG(LogTemp, Warning, TEXT("[FSessionsCommon::CheckCreateSessionState] Could not create session with bPresenceEnabled set to true when another already exists [%s]."), *Entry.Key.ToString());
-
-					return Errors::InvalidState();
-				}
-			}
-		}
+		TSharedRef<FSessionCommon> FoundSession = GetMutableSessionByNameResult.GetOkValue().Session;
 
-		// User login check for all local users
-		IAuthPtr Auth = Services.GetAuthInterface();
-		if (!Auth->IsLoggedIn(Params.LocalAccountId))
-		{
-			UE_LOG(LogTemp, Warning, TEXT("[FSessionsCommon::CheckCreateSessionState] Could not create session with user [%s] not logged in"), *ToLogString(Params.LocalAccountId));
+		FSessionSettings& SessionSettings = FoundSession->SessionSettings;
 
-			return Errors::InvalidUser();
-		}
+		FoundSession->SessionMembers.Remove(Params.LocalAccountId);
 
-		return Errors::Success();
+		return TOnlineResult<FRemoveSessionMember>();
 	}
 
-	FOnlineError FSessionsCommon::CheckUpdateSessionSettingsParams(const FUpdateSessionSettings::Params& Params)
+	TOptional<FOnlineError> FSessionsCommon::CheckParams(const FRemoveSessionMember::Params& Params) const
 	{
-		if (!Params.LocalAccountId.IsValid())
-		{
-			return Errors::InvalidUser();
-		}
+		CHECK_PARAMS_ID_HANDLE(Params.LocalAccountId)
+		CHECK_PARAMS_FNAME(Params.SessionName)
 
-		if (Params.SessionName.IsNone())
-		{
-			return Errors::InvalidParams();
-		}
-
-		for (const TPair<FSchemaAttributeId, FCustomSessionSetting>& Entry : Params.Mutations.UpdatedCustomSettings)
-		{
-			if (Entry.Key.IsNone())
-			{
-				return Errors::InvalidParams();
-			}
-		}
-
-		return Errors::Success();
+		return TOptional<FOnlineError>();
 	}
 
-	FOnlineError FSessionsCommon::CheckUpdateSessionSettingsState(const FUpdateSessionSettings::Params& Params)
+	TOptional<FOnlineError> FSessionsCommon::CheckState(const FRemoveSessionMember::Params& Params) const
 	{
-		// User login check
-		IAuthPtr Auth = Services.GetAuthInterface();
-		if (!Auth->IsLoggedIn(Params.LocalAccountId))
-		{
-			UE_LOG(LogTemp, Warning, TEXT("[FSessionsCommon::CheckUpdateSessionSettingsState] User [%s] not logged in"), *ToLogString(Params.LocalAccountId));
+		CHECK_STATE_EXISTING_SESSION_NAME(Params.SessionName)
 
-			return Errors::InvalidUser();
-		}
+		return TOptional<FOnlineError>();
+	}
 
-		// Session name check
-		if (TOptional<FOnlineError> Result = CheckSessionExistsByName(Params.LocalAccountId, Params.SessionName))
-		{
-			UE_LOG(LogTemp, Warning, TEXT("[FSessionsCommon::CheckUpdateSessionSettingsState] Session with name [%s] not found."), *Params.SessionName.ToString());
+	// SendSessionInvite
 
-			return Result.GetValue();
-		}
+	TOnlineAsyncOpHandle<FSendSessionInvite> FSessionsCommon::SendSessionInvite(FSendSessionInvite::Params&& Params)
+	{
+		return ExecuteAsyncSessionsMethod<FSendSessionInvite>(MoveTemp(Params), &FSessionsCommon::SendSessionInviteImpl);
+	}
 
-		return Errors::Success();
+	TFuture<TOnlineResult<FSendSessionInvite>> FSessionsCommon::SendSessionInviteImpl(const FSendSessionInvite::Params& Params)
+	{
+		return MakeFulfilledPromise<TOnlineResult<FSendSessionInvite>>(Errors::NotImplemented()).GetFuture();
 	}
 
-	FOnlineError FSessionsCommon::CheckFindSessionsParams(const FFindSessions::Params& Params)
+	TOptional<FOnlineError> FSessionsCommon::CheckParams(const FSendSessionInvite::Params& Params) const
 	{
-		if (Params.MaxResults <= 0)
+		CHECK_PARAMS_ID_HANDLE(Params.LocalAccountId)
+		CHECK_PARAMS_FNAME(Params.SessionName)
+
+		if (Params.TargetUsers.IsEmpty())
 		{
-			UE_LOG(LogTemp, Warning, TEXT("[FSessionsCommon::CheckFindSessionsParams] Could not find sessions with no valid MaxResults [%d]"), Params.MaxResults);
+			UE_LOG(LogTemp, Warning, TEXT("[%s] TargetUsers array is empty"), UTF8_TO_TCHAR(__FUNCTION__));
+			return TOptional<FOnlineError>(Errors::InvalidParams());
+		}
 
-			return Errors::InvalidParams();
+		for (const FAccountId& TargetUser : Params.TargetUsers)
+		{
+			CHECK_PARAMS_ID_HANDLE(TargetUser)
 		}
 
-		return Errors::Success();
+		return TOptional<FOnlineError>();
 	}
 
-	FOnlineError FSessionsCommon::CheckFindSessionsState(const FFindSessions::Params& Params)
+	TOptional<FOnlineError> FSessionsCommon::CheckState(const FSendSessionInvite::Params& Params) const
 	{
-		// User login check
-		if (!Services.GetAuthInterface()->IsLoggedIn(Params.LocalAccountId))
-		{
-			UE_LOG(LogTemp, Warning, TEXT("[FSessionsCommon::CheckFindSessionsState] Could not find sessions with user [%s] not logged in"), *ToLogString(Params.LocalAccountId));
-
-			return Errors::InvalidUser();
-		}
+		CHECK_STATE_EXISTING_SESSION_NAME(Params.SessionName)
 
-		// Ongoing search check
-		if (const TSharedRef<TOnlineAsyncOp<FFindSessions>>* CurrentSessionSearchHandle = CurrentSessionSearchHandlesUserMap.Find(Params.LocalAccountId))
-		{
-			UE_LOG(LogTemp, Warning, TEXT("[FSessionsCommon::CheckFindSessionsState] Could not find sessions, search already in progress"));
+		return TOptional<FOnlineError>();
+	}
 
-			return Errors::AlreadyPending();
-		}
+	// RejectSessionInvite
 
-		return Errors::Success();
+	TOnlineAsyncOpHandle<FRejectSessionInvite> FSessionsCommon::RejectSessionInvite(FRejectSessionInvite::Params&& Params)
+	{
+		return ExecuteAsyncSessionsMethod<FRejectSessionInvite>(MoveTemp(Params), &FSessionsCommon::RejectSessionInviteImpl);
 	}
 
-	FOnlineError FSessionsCommon::CheckStartMatchmakingParams(const FStartMatchmaking::Params& Params)
+	TFuture<TOnlineResult<FRejectSessionInvite>> FSessionsCommon::RejectSessionInviteImpl(const FRejectSessionInvite::Params& Params)
 	{
-		if (Params.SessionCreationParameters.SessionSettings.NumMaxConnections == 0)
-		{
-			UE_LOG(LogTemp, Warning, TEXT("[FSessionsCommon::CheckStartMatchmakingParams] Could not create session with invalid NumMaxPrivateConnections and NumMaxPublicConnections"));
+		return MakeFulfilledPromise<TOnlineResult<FRejectSessionInvite>>(Errors::NotImplemented()).GetFuture();
+	}
 
-			return Errors::InvalidParams();
-		}
+	TOptional<FOnlineError> FSessionsCommon::CheckParams(const FRejectSessionInvite::Params& Params) const
+	{
+		CHECK_PARAMS_ID_HANDLE(Params.LocalAccountId)
+		CHECK_PARAMS_ID_HANDLE(Params.SessionInviteId)
 
-		return Errors::Success();
+		return TOptional<FOnlineError>();
 	}
 
-	FOnlineError FSessionsCommon::CheckStartMatchmakingState(const FStartMatchmaking::Params& Params)
+	TOptional<FOnlineError> FSessionsCommon::CheckState(const FRejectSessionInvite::Params& Params) const
 	{
-		// Check if a session with that name already exists
-		TOptional<FOnlineError> Result = CheckSessionExistsByName(Params.SessionCreationParameters.LocalAccountId, Params.SessionCreationParameters.SessionName);
-		if (!Result.IsSet()) // If CheckSessionExistsByName did not return an error, a session with that name already exists
+		if (const TMap<FSessionInviteId, TSharedRef<FSessionInvite>>* UserMap = SessionInvitesUserMap.Find(Params.LocalAccountId))
 		{
-			UE_LOG(LogTemp, Warning, TEXT("[FSessionsCommon::CheckStartMatchmakingState] Could not join session with name [%s]. A session with that name already exists"), *Params.SessionCreationParameters.SessionName.ToString());
-
-			return Errors::InvalidState(); // TODO: New error: Session with name %s already exists
+			if (!UserMap->Contains(Params.SessionInviteId))
+			{
+				UE_LOG(LogTemp, Warning, TEXT("[%s] No invite found for invite id [%s]."), UTF8_TO_TCHAR(__FUNCTION__), *ToLogString(Params.SessionInviteId));
+				return TOptional<FOnlineError>(Errors::NotFound());
+			}
 		}
-
-		// User login check for all local users
-		IAuthPtr Auth = Services.GetAuthInterface();
-		if (!Auth->IsLoggedIn(Params.SessionCreationParameters.LocalAccountId))
+		else
 		{
-			UE_LOG(LogTemp, Warning, TEXT("[FSessionsCommon::CheckStartMatchmakingState] Could not join session with user [%s] not logged in"), *ToLogString(Params.SessionCreationParameters.LocalAccountId));
-
-			return Errors::InvalidUser();
+			UE_LOG(LogTemp, Warning, TEXT("[%s] No invite data found for account id [%s]."), UTF8_TO_TCHAR(__FUNCTION__), *ToLogString(Params.LocalAccountId));
+			return TOptional<FOnlineError>(Errors::InvalidState());
 		}
 
-		// TODO: Check that only one session has bUsesPresence set
-
-		return Errors::Success();
+		return TOptional<FOnlineError>();
 	}
 
-	FOnlineError FSessionsCommon::CheckJoinSessionParams(const FJoinSession::Params& Params)
-	{
-		if (!Params.LocalAccountId.IsValid())
-		{
-			UE_LOG(LogTemp, Warning, TEXT("[FSessionsCommon::CheckJoinSessionParams] Could not join session with name [%s]. LocalAccountId [%s] not valid"), *ToLogString(Params.LocalAccountId));
+#undef CHECK_PARAMS_ID_HANDLE
+#undef CHECK_PARAMS_FNAME
+#undef CHECK_STATE_SESSION_NAME_ALREADY_EXISTS
 
-			return Errors::InvalidParams();
-		}
+	// Events
 
-		if (!Params.SessionId.IsValid())
-		{
-			UE_LOG(LogTemp, Warning, TEXT("[FSessionsCommon::CheckJoinSessionParams] Could not join session with name [%s]. SessionId [%s] not valid"), *ToLogString(Params.SessionId));
-
-			return Errors::InvalidParams();
-		}
+	TOnlineEvent<void(const FSessionJoined&)> FSessionsCommon::OnSessionJoined()
+	{
+		return SessionEvents.OnSessionJoined;
+	}
 
-		if (Params.SessionName.IsNone())
-		{
-			UE_LOG(LogTemp, Warning, TEXT("[FSessionsCommon::CheckJoinSessionParams] Could not join session with no valid SessionName set"));
+	TOnlineEvent<void(const FSessionLeft&)> FSessionsCommon::OnSessionLeft()
+	{
+		return SessionEvents.OnSessionLeft;
+	}
 
-			return Errors::InvalidParams();
-		}
+	TOnlineEvent<void(const FSessionUpdated&)> FSessionsCommon::OnSessionUpdated()
+	{
+		return SessionEvents.OnSessionUpdated;
+	}
 
-		return Errors::Success();
+	TOnlineEvent<void(const FSessionInviteReceived&)> FSessionsCommon::OnSessionInviteReceived()
+	{
+		return SessionEvents.OnSessionInviteReceived;
 	}
 
-	FOnlineError FSessionsCommon::CheckJoinSessionState(const FJoinSession::Params& Params)
+	TOnlineEvent<void(const FUISessionJoinRequested&)> FSessionsCommon::OnUISessionJoinRequested()
 	{
-		// Check if a session with that name already exists
-		TOptional<FOnlineError> Result = CheckSessionExistsByName(Params.LocalAccountId, Params.SessionName);
-		if (!Result.IsSet()) // If CheckSessionExistsByName did not return an error, a session with that name already exists
-		{
-			UE_LOG(LogTemp, Warning, TEXT("[FSessionsCommon::CheckJoinSessionState] Could not join session with name [%s]. A session with that name already exists"), *Params.SessionName.ToString());
+		return SessionEvents.OnUISessionJoinRequested;
+	}
 
-			return Errors::InvalidState(); // TODO: New error: Session with name %s already exists
-		}
+	// Auxiliary methods
 
-		// User login check for all local users
-		IAuthPtr Auth = Services.GetAuthInterface();
-		if(!Auth->IsLoggedIn(Params.LocalAccountId))
-		{
-			UE_LOG(LogTemp, Warning, TEXT("[FSessionsCommon::CheckJoinSessionState] Could not join session with user [%s] not logged in"), *ToLogString(Params.LocalAccountId));
+	void FSessionsCommon::AddSessionInvite(const TSharedRef<FSessionInvite> SessionInvite, const TSharedRef<FSessionCommon> Session, const FAccountId& LocalAccountId)
+	{
+		AllSessionsById.Emplace(Session->GetSessionId(), Session);
 
-			return Errors::InvalidUser();
-		}
+		TMap<FSessionInviteId, TSharedRef<FSessionInvite>>& SessionInvitesMap = SessionInvitesUserMap.FindOrAdd(LocalAccountId);
+		SessionInvitesMap.Emplace(SessionInvite->InviteId, SessionInvite);
+	}
 
-		// We check that the session is cached and valid for a join operation by the users
-		TOnlineResult<FGetSessionById> GetSessionByIdResult = GetSessionById({ Params.SessionId });
-		if (GetSessionByIdResult.IsError())
-		{
-			UE_LOG(LogTemp, Warning, TEXT("[FSessionsCommon::CheckJoinSessionState] Unable to find session with SessionId parameter [%s]. Please call FindSessions to get an updated list of available sessions "), *ToLogString(Params.SessionId));
+	void FSessionsCommon::AddSearchResult(const TSharedRef<FSessionCommon> Session, const FAccountId& LocalAccountId)
+	{
+		AllSessionsById.Emplace(Session->GetSessionId(), Session);
 
-			return GetSessionByIdResult.GetErrorValue();
-		}
+		TArray<FOnlineSessionId>& SearchResults = SearchResultsUserMap.FindOrAdd(LocalAccountId);
+		SearchResults.Add(Session->GetSessionId());
+	}
 
-		TSharedRef<const ISession> FoundSession = GetSessionByIdResult.GetOkValue().Session;
+	void FSessionsCommon::AddSessionWithReferences(const TSharedRef<FSessionCommon> Session, const FName& SessionName, const FAccountId& LocalAccountId, bool bIsPresenceSession)
+	{
+		AllSessionsById.Emplace(Session->GetSessionId(), Session);
 
-		const FSessionSettings& SessionSettings = FoundSession->GetSessionSettings();
+		AddSessionReferences(Session->GetSessionId(), SessionName, LocalAccountId, bIsPresenceSession);
+	}
 
-		if (FoundSession->GetSessionMembers().Contains(Params.LocalAccountId))
-		{
-			UE_LOG(LogTemp, Warning, TEXT("[FSessionsCommon::CheckJoinSessionState] Could not join session. User [%s] already in session"), *ToLogString(Params.LocalAccountId));
+	void FSessionsCommon::AddSessionReferences(const FOnlineSessionId SessionId, const FName& SessionName, const FAccountId& LocalAccountId, bool bIsPresenceSession)
+	{
+		LocalSessionsByName.Emplace(SessionName, SessionId);
 
-			return Errors::AccessDenied();
-		}
+		NamedSessionUserMap.FindOrAdd(LocalAccountId).AddUnique(SessionName);
 
-		if (!FoundSession->IsJoinable())
+		if (bIsPresenceSession)
 		{
-			UE_LOG(LogTemp, Warning, TEXT("[FSessionsCommon::CheckJoinSessionState] Could not join session. Session not joinable "));
-
-			return Errors::AccessDenied();
+			SetPresenceSession({ LocalAccountId, SessionId });
 		}
+	}
 
-		if (Params.bPresenceEnabled)
+	void FSessionsCommon::ClearSessionInvitesForSession(const FAccountId& LocalAccountId, const FOnlineSessionId SessionId)
+	{
+		if (TMap<FSessionInviteId, TSharedRef<FSessionInvite>>* UserMap = SessionInvitesUserMap.Find(LocalAccountId))
 		{
-			for (const TPair<FName, FOnlineSessionId>& Entry : LocalSessionsByName)
+			TArray<FSessionInviteId> InviteIdsToRemove;
+			for (const TPair<FSessionInviteId, TSharedRef<FSessionInvite>>& Entry : *UserMap)
 			{
-				TOnlineResult<FGetPresenceSession> GetPresenceSessionResult = GetPresenceSession({ Params.LocalAccountId });
-				if (GetPresenceSessionResult.IsOk())
+				if (Entry.Value->SessionId == SessionId)
 				{
-					UE_LOG(LogTemp, Warning, TEXT("[FSessionsCommon::CheckJoinSessionState] Could not create session with bPresenceEnabled set to true when another already exists [%s]."), *Entry.Key.ToString());
-
-					return Errors::InvalidState();
+					InviteIdsToRemove.Add(Entry.Key);
 				}
 			}
-		}
 
-		return Errors::Success();
+			for (const FSessionInviteId& InviteId : InviteIdsToRemove)
+			{
+				UserMap->Remove(InviteId);
+			}
+		}
 	}
 
-	FOnlineError FSessionsCommon::CheckAddSessionMemberState(const FAddSessionMember::Params& Params)
+	void FSessionsCommon::ClearSessionReferences(const FOnlineSessionId SessionId, const FName& SessionName, const FAccountId& LocalAccountId)
 	{
-		if (TOptional<FOnlineError> Result = CheckSessionExistsByName(Params.LocalAccountId, Params.SessionName))
-		{
-			UE_LOG(LogTemp, Warning, TEXT("[FSessionsCommon::CheckAddSessionMemberState] Could not add session member to session with name [%s]. Session not found"), *Params.SessionName.ToString());
+		NamedSessionUserMap.FindChecked(LocalAccountId).Remove(SessionName);
 
-			return Result.GetValue();
+		TOnlineResult<FIsPresenceSession> IsPresenceSessionResult = IsPresenceSession({ LocalAccountId, SessionId });
+		if (IsPresenceSessionResult.IsOk())
+		{
+			if (IsPresenceSessionResult.GetOkValue().bIsPresenceSession)
+			{
+				ClearPresenceSession({ LocalAccountId });
+			}
 		}
 
-		// TODO: Check if there are enough slots available if Params::bReserveSlot is true
-
-		return Errors::Success();
+		ClearSessionByName(SessionName);
+		ClearSessionById(SessionId);
 	}
 
-	FOnlineError FSessionsCommon::CheckRemoveSessionMemberState(const FRemoveSessionMember::Params& Params)
+	void FSessionsCommon::ClearSessionByName(const FName& SessionName)
 	{
-		if (TOptional<FOnlineError> Result = CheckSessionExistsByName(Params.LocalAccountId, Params.SessionName))
+		for (const TPair<FAccountId, TArray<FName>>& Entry : NamedSessionUserMap)
 		{
-			UE_LOG(LogTemp, Warning, TEXT("[FSessionsCommon::CheckRemoveSessionMemberState] Could not remove session member from session with name [%s]. Session not found"), *Params.SessionName.ToString());
-
-			return Result.GetValue();
+			if (Entry.Value.Contains(SessionName))
+			{
+				return;
+			}
 		}
 
-		return Errors::Success();
+		// If no references were found, we'll remove the named session entry
+		LocalSessionsByName.Remove(SessionName);
 	}
 
-	FOnlineError FSessionsCommon::CheckLeaveSessionState(const FLeaveSession::Params& Params)
+	void FSessionsCommon::ClearSessionById(const FOnlineSessionId& SessionId)
 	{
-		// User login check for main caller, session check
-		if (TOptional<FOnlineError> Result = CheckSessionExistsByName(Params.LocalAccountId, Params.SessionName))
-		{
-			UE_LOG(LogTemp, Warning, TEXT("[FSessionsCommon::CheckLeaveSessionState] Could not leave session with name [%s]. Session not found"), *Params.SessionName.ToString());
-
-			return Result.GetValue();
-		}
-
-		// User login check
-		IAuthPtr Auth = Services.GetAuthInterface();
-		if (!Auth->IsLoggedIn(Params.LocalAccountId))
+		// PresenceSessionsUserMap is not evaluated, since any session there would also be in LocalSessionsByName
+		for (const TPair<FName, FOnlineSessionId>& Entry : LocalSessionsByName)
 		{
-			UE_LOG(LogTemp, Warning, TEXT("[FSessionsCommon::CheckLeaveSessionState] Could not leave session with user [%s] not logged in"), *ToLogString(Params.LocalAccountId));
-
-			return Errors::InvalidUser();
+			if (Entry.Value == SessionId)
+			{
+				return;
+			}
 		}
 
-		return Errors::Success();
-	}
-
-	FOnlineError FSessionsCommon::CheckSendSessionInviteState(const FSendSessionInvite::Params& Params)
-	{
-		// User login check for main caller, session check
-		TOnlineResult<FGetSessionByName> GetSessionByNameResult = GetSessionByName({ Params.SessionName });
-		if (GetSessionByNameResult.IsOk())
+		for (const TPair<FAccountId, TMap<FSessionInviteId, TSharedRef<FSessionInvite>>>& Entry : SessionInvitesUserMap)
 		{
-			TSharedRef<const ISession> FoundSession = GetSessionByNameResult.GetOkValue().Session;
-
-			if (!FoundSession->GetSessionInfo().bIsDedicatedServerSession)
+			for (const TPair<FSessionInviteId, TSharedRef<FSessionInvite>>& InviteMap : Entry.Value)
 			{
-				if (!Services.GetAuthInterface()->IsLoggedIn(Params.LocalAccountId))
+				if (InviteMap.Value->SessionId == SessionId)
 				{
-					UE_LOG(LogTemp, Warning, TEXT("[FSessionsNull::CheckSendSessionInviteState] Could not send session invite with user [%s] not logged in"), *ToLogString(Params.LocalAccountId));
-
-					return Errors::InvalidUser();
+					return;
 				}
-			}
+			}			
 		}
-		else
-		{
-			UE_LOG(LogTemp, Warning, TEXT("[FSessionsCommon::CheckSendSessionInviteState] Could not send session invite for session with name [%s]. Session not found"), *Params.SessionName.ToString());
 
-			return Errors::NotFound();
+		for (const TPair<FAccountId, TArray<FOnlineSessionId>>& Entry : SearchResultsUserMap)
+		{
+			if (Entry.Value.Contains(SessionId))
+			{
+				return;
+			}
 		}
 
-		// TODO: check if Session Invite Id is valid
-
-		return Errors::Success();
+		// If no references were found, we'll remove the session entry
+		AllSessionsById.Remove(SessionId);
 	}
 
-	FOnlineError FSessionsCommon::CheckRejectSessionInviteState(const FRejectSessionInvite::Params& Params)
-	{
-		if (!Services.GetAuthInterface()->IsLoggedIn(Params.LocalAccountId))
-		{
-			UE_LOG(LogTemp, Warning, TEXT("[FSessionsNull::CheckRejectSessionInviteState] Could not send session invite with user [%s] not logged in"), *ToLogString(Params.LocalAccountId));
+#define COPY_TOPTIONAL_VALUE_IF_SET(Value) \
+	if (UpdatedValues.Value.IsSet()) \
+	{ \
+		SessionSettingsChanges.Value = UpdatedValues.Value.GetValue(); \
+	} \
 
-			return Errors::InvalidUser();
-		}
+	FSessionUpdate FSessionsCommon::BuildSessionUpdate(const TSharedRef<FSessionCommon>& Session, const FSessionSettingsUpdate& UpdatedValues) const
+	{
+		FSessionUpdate Result;
 
-		// TODO: check if Session Invite Id is valid
+		FSessionSettingsChanges SessionSettingsChanges;
 
-		return Errors::Success();
-	}
+		COPY_TOPTIONAL_VALUE_IF_SET(SchemaName)
+		COPY_TOPTIONAL_VALUE_IF_SET(NumMaxConnections)
+		COPY_TOPTIONAL_VALUE_IF_SET(JoinPolicy)
+		COPY_TOPTIONAL_VALUE_IF_SET(bAllowNewMembers)
 
-	// TODO: Have all Check methods return TOptional too, change call sites, and write Macro for repeating code structure
-	TOptional<FOnlineError> FSessionsCommon::CheckSessionExistsByName(const FAccountId& LocalAccountId, const FName& SessionName)
-	{
-		TOptional<FOnlineError> Result;
+		SessionSettingsChanges.RemovedCustomSettings.Append(UpdatedValues.RemovedCustomSettings);
 
-		TOnlineResult<FGetSessionByName> GetSessionByNameResult = GetSessionByName({ SessionName });
-		if (GetSessionByNameResult.IsError())
+		for (const TPair<FSchemaAttributeId, FCustomSessionSetting>& Entry : UpdatedValues.UpdatedCustomSettings)
 		{
-			Result.Emplace(GetSessionByNameResult.GetErrorValue());
+			if (const FCustomSessionSetting* CustomSetting = Session->GetSessionSettings().CustomSettings.Find(Entry.Key))
+			{
+				FCustomSessionSettingUpdate SettingUpdate = { *CustomSetting, Entry.Value };
+
+				SessionSettingsChanges.ChangedCustomSettings.Emplace(Entry.Key, SettingUpdate);
+			}
+			else
+			{
+				SessionSettingsChanges.AddedCustomSettings.Add(Entry);
+			}
 		}
 
+		Result.SessionSettingsChanges = SessionSettingsChanges;
+
 		return Result;
 	}
 
+#undef COPY_TOPTIONAL_VALUE_IF_SET
+
 /* UE::Online */ }
\ No newline at end of file
diff --git a/Engine/Plugins/Online/Experimental/OnlineServices/Source/OnlineServicesCommon/Public/Online/SessionsCommon.h b/Engine/Plugins/Online/Experimental/OnlineServices/Source/OnlineServicesCommon/Public/Online/SessionsCommon.h
--- a/Engine/Plugins/Online/Experimental/OnlineServices/Source/OnlineServicesCommon/Public/Online/SessionsCommon.h
+++ b/Engine/Plugins/Online/Experimental/OnlineServices/Source/OnlineServicesCommon/Public/Online/SessionsCommon.h
@@ -205,28 +205,6 @@ struct FGetMutableSessionById
 	};
 };
 
-struct FUpdateSessionSettingsImpl
-{
-	static constexpr TCHAR Name[] = TEXT("UpdateSessionSettingsImpl");
-
-	struct Params
-	{
-		/** The local user agent which will perform the action. */
-		FAccountId LocalAccountId;
-
-		/** The local name for the session */
-		FName SessionName;
-
-		/** Changes to current session settings */
-		FSessionSettingsUpdate Mutations;
-	};
-
-	struct Result
-	{
-
-	};
-};
-
 class ONLINESERVICESCOMMON_API FSessionsCommon : public TOnlineComponent<ISessions>
 {
 public:
@@ -266,6 +244,7 @@ public:
 	virtual TOnlineEvent<void(const FUISessionJoinRequested&)> OnUISessionJoinRequested() override;
 
 protected:
+	// Auxiliary methods
 	TOnlineResult<FGetMutableSessionByName> GetMutableSessionByName(FGetMutableSessionByName::Params&& Params) const;
 	TOnlineResult<FGetMutableSessionById> GetMutableSessionById(FGetMutableSessionById::Params&& Params) const;
 
@@ -279,41 +258,67 @@ protected:
 
 	FSessionUpdate BuildSessionUpdate(const TSharedRef<FSessionCommon>& Session, const FSessionSettingsUpdate& UpdatedValues) const;
 
-	virtual TFuture<TOnlineResult<FUpdateSessionSettingsImpl>> UpdateSessionSettingsImpl(FUpdateSessionSettingsImpl::Params&& Params);
-	TOnlineResult<FAddSessionMember> AddSessionMemberImpl(const FAddSessionMember::Params& Params);
-	TOnlineResult<FRemoveSessionMember> RemoveSessionMemberImpl(const FRemoveSessionMember::Params& Params);
+	// FSessionsCommon internal interface
 
-	FOnlineError CheckCreateSessionParams(const FCreateSession::Params& Params);
-	FOnlineError CheckCreateSessionState(const FCreateSession::Params& Params);
+	virtual TFuture<TOnlineResult<FCreateSession>> CreateSessionImpl(const FCreateSession::Params& Params);
 
-	FOnlineError CheckUpdateSessionSettingsParams(const FUpdateSessionSettings::Params& Params);
-	FOnlineError CheckUpdateSessionSettingsState(const FUpdateSessionSettings::Params& Params);
+	virtual TFuture<TOnlineResult<FUpdateSessionSettings>> UpdateSessionSettingsImpl(const FUpdateSessionSettings::Params& Params);
 
-	FOnlineError CheckFindSessionsParams(const FFindSessions::Params& Params);
-	FOnlineError CheckFindSessionsState(const FFindSessions::Params& Params);
+	virtual TFuture<TOnlineResult<FLeaveSession>> LeaveSessionImpl(const FLeaveSession::Params& Params);
 
-	FOnlineError CheckStartMatchmakingParams(const FStartMatchmaking::Params& Params);
-	FOnlineError CheckStartMatchmakingState(const FStartMatchmaking::Params& Params);
+	virtual TFuture<TOnlineResult<FFindSessions>> FindSessionsImpl(const FFindSessions::Params& Params);
 
-	FOnlineError CheckJoinSessionParams(const FJoinSession::Params& Params);
-	FOnlineError CheckJoinSessionState(const FJoinSession::Params& Params);
+	virtual TFuture<TOnlineResult<FStartMatchmaking>> StartMatchmakingImpl(const FStartMatchmaking::Params& Params);
 
-	FOnlineError CheckAddSessionMemberState(const FAddSessionMember::Params& Params);
+	virtual TFuture<TOnlineResult<FJoinSession>> JoinSessionImpl(const FJoinSession::Params& Params);
 
-	FOnlineError CheckRemoveSessionMemberState(const FRemoveSessionMember::Params& Params);
+	virtual TFuture<TOnlineResult<FAddSessionMember>> AddSessionMemberImpl(const FAddSessionMember::Params& Params);
+	TOnlineResult<FAddSessionMember> AddSessionMemberInternal(const FAddSessionMember::Params& Params);
 
-	FOnlineError CheckLeaveSessionState(const FLeaveSession::Params& Params);
+	virtual TFuture<TOnlineResult<FRemoveSessionMember>> RemoveSessionMemberImpl(const FRemoveSessionMember::Params& Params);
+	TOnlineResult<FRemoveSessionMember> RemoveSessionMemberInternal(const FRemoveSessionMember::Params& Params);
 
-	FOnlineError CheckSendSessionInviteState(const FSendSessionInvite::Params& Params);
+	virtual TFuture<TOnlineResult<FSendSessionInvite>> SendSessionInviteImpl(const FSendSessionInvite::Params& Params);
 
-	FOnlineError CheckRejectSessionInviteState(const FRejectSessionInvite::Params& Params);
+	virtual TFuture<TOnlineResult<FRejectSessionInvite>> RejectSessionInviteImpl(const FRejectSessionInvite::Params& Params);
 
 private:
-	TOptional<FOnlineError> CheckSessionExistsByName(const FAccountId& LocalAccountId, const FName& SessionName);
+	template<typename MethodStruct>
+	TOnlineAsyncOpHandle<MethodStruct> ExecuteAsyncSessionsMethod(typename MethodStruct::Params&& Params, TFuture<TOnlineResult<MethodStruct>>(FSessionsCommon::* ImplFunc)(const typename MethodStruct::Params& Params));
 
 	void ClearSessionByName(const FName& SessionName);
 	void ClearSessionById(const FOnlineSessionId& SessionId);
 
+	TOptional<FOnlineError> CheckParams(const FCreateSession::Params& Params) const;
+	TOptional<FOnlineError> CheckState(const FCreateSession::Params& Params) const;
+
+	TOptional<FOnlineError> CheckParams(const FUpdateSessionSettings::Params& Params) const;
+	TOptional<FOnlineError> CheckState(const FUpdateSessionSettings::Params& Params) const;
+
+	TOptional<FOnlineError> CheckParams(const FLeaveSession::Params& Params) const;
+	TOptional<FOnlineError> CheckState(const FLeaveSession::Params& Params) const;
+
+	TOptional<FOnlineError> CheckParams(const FFindSessions::Params& Params) const;
+	TOptional<FOnlineError> CheckState(const FFindSessions::Params& Params) const;
+
+	TOptional<FOnlineError> CheckParams(const FStartMatchmaking::Params& Params) const;
+	TOptional<FOnlineError> CheckState(const FStartMatchmaking::Params& Params) const;
+
+	TOptional<FOnlineError> CheckParams(const FJoinSession::Params& Params) const;
+	TOptional<FOnlineError> CheckState(const FJoinSession::Params& Params) const;
+
+	TOptional<FOnlineError> CheckParams(const FAddSessionMember::Params& Params) const;
+	TOptional<FOnlineError> CheckState(const FAddSessionMember::Params& Params) const;
+
+	TOptional<FOnlineError> CheckParams(const FRemoveSessionMember::Params& Params) const;
+	TOptional<FOnlineError> CheckState(const FRemoveSessionMember::Params& Params) const;
+
+	TOptional<FOnlineError> CheckParams(const FSendSessionInvite::Params& Params) const;
+	TOptional<FOnlineError> CheckState(const FSendSessionInvite::Params& Params) const;
+
+	TOptional<FOnlineError> CheckParams(const FRejectSessionInvite::Params& Params) const;
+	TOptional<FOnlineError> CheckState(const FRejectSessionInvite::Params& Params) const;
+
 protected:
 
 	struct FSessionEvents
@@ -341,7 +346,7 @@ protected:
 	TMap<FAccountId, TArray<FOnlineSessionId>> SearchResultsUserMap;
 
 	/** Handle for an ongoing session search operation, mapped per user */
-	TMap<FAccountId, TSharedRef<TOnlineAsyncOp<FFindSessions>>> CurrentSessionSearchHandlesUserMap;
+	TMap<FAccountId, TPromise<TOnlineResult<FFindSessions>>> CurrentSessionSearchPromisesUserMap;
 
 	/** Set of every distinct FSession found, indexed by Id */
 	TMap<FOnlineSessionId, TSharedRef<FSessionCommon>> AllSessionsById;
@@ -365,15 +370,6 @@ BEGIN_ONLINE_STRUCT_META(FGetMutableSessionById::Result)
 	ONLINE_STRUCT_FIELD(FGetMutableSessionById::Result, Session)
 END_ONLINE_STRUCT_META()
 
-BEGIN_ONLINE_STRUCT_META(FUpdateSessionSettingsImpl::Params)
-	ONLINE_STRUCT_FIELD(FUpdateSessionSettingsImpl::Params, LocalAccountId),
-	ONLINE_STRUCT_FIELD(FUpdateSessionSettingsImpl::Params, SessionName),
-	ONLINE_STRUCT_FIELD(FUpdateSessionSettingsImpl::Params, Mutations)
-END_ONLINE_STRUCT_META()
-
-BEGIN_ONLINE_STRUCT_META(FUpdateSessionSettingsImpl::Result)
-END_ONLINE_STRUCT_META()
-
 /* Meta*/ }
 
 /* UE::Online */ }
diff --git a/Engine/Plugins/Online/Experimental/OnlineServices/Source/OnlineServicesCommonEngineUtils/Private/Online/SessionsLAN.cpp b/Engine/Plugins/Online/Experimental/OnlineServices/Source/OnlineServicesCommonEngineUtils/Private/Online/SessionsLAN.cpp
--- a/Engine/Plugins/Online/Experimental/OnlineServices/Source/OnlineServicesCommonEngineUtils/Private/Online/SessionsLAN.cpp
+++ b/Engine/Plugins/Online/Experimental/OnlineServices/Source/OnlineServicesCommonEngineUtils/Private/Online/SessionsLAN.cpp
@@ -124,66 +124,46 @@ void FSessionsLAN::Tick(float DeltaSeconds)
 	Super::Tick(DeltaSeconds);
 }
 
-TOnlineAsyncOpHandle<FCreateSession> FSessionsLAN::CreateSession(FCreateSession::Params&& Params)
+TFuture<TOnlineResult<FCreateSession>> FSessionsLAN::CreateSessionImpl(const FCreateSession::Params& Params)
 {
-	TOnlineAsyncOpRef<FCreateSession> Op = GetOp<FCreateSession>(MoveTemp(Params));
-	const FCreateSession::Params& OpParams = Op->GetParams();
+	TPromise<TOnlineResult<FCreateSession>> Promise;
+	TFuture<TOnlineResult<FCreateSession>> Future = Promise.GetFuture();
 
-	FOnlineError ParamsCheck = CheckCreateSessionParams(OpParams);
-	if (ParamsCheck != Errors::Success())
+	// We'll only host on the LAN beacon for public sessions
+	if (Params.SessionSettings.JoinPolicy == ESessionJoinPolicy::Public)
 	{
-		Op->SetError(MoveTemp(ParamsCheck));
-		return Op->GetHandle();
-	}
-
-	Op->Then([this](TOnlineAsyncOp<FCreateSession>& Op)
-	{
-		const FCreateSession::Params& OpParams = Op.GetParams();
-
-		FOnlineError StateCheck = CheckCreateSessionState(OpParams);
-		if (StateCheck != Errors::Success())
-		{
-			Op.SetError(MoveTemp(StateCheck));
-			return;
-		}
-
-		// We'll only host on the LAN beacon for public sessions
-		if (OpParams.SessionSettings.JoinPolicy == ESessionJoinPolicy::Public)
+		if (TOptional<FOnlineError> TryHostLANSessionResult = TryHostLANSession())
 		{
-			if (TOptional<FOnlineError> TryHostLANSessionResult = TryHostLANSession())
-			{
-				Op.SetError(MoveTemp(TryHostLANSessionResult.GetValue()));
-				return;
-			}
+			Promise.EmplaceValue(TryHostLANSessionResult.GetValue());
+			return Future;
 		}
+	}
 
-		TSharedRef<FSessionLAN> NewSessionLANRef = MakeShared<FSessionLAN>();
-		NewSessionLANRef->OwnerAccountId = OpParams.LocalAccountId;
-		NewSessionLANRef->SessionInfo.bAllowSanctionedPlayers = OpParams.bAllowSanctionedPlayers;
-		NewSessionLANRef->SessionInfo.bAntiCheatProtected = OpParams.bAntiCheatProtected;
-		NewSessionLANRef->SessionInfo.bIsDedicatedServerSession = IsRunningDedicatedServer();
-		NewSessionLANRef->SessionInfo.bIsLANSession = true;
-		NewSessionLANRef->SessionInfo.SessionId = FOnlineSessionIdRegistryLAN::GetChecked(Services.GetServicesProvider()).GetNextSessionId();
-		NewSessionLANRef->SessionInfo.SessionIdOverride = OpParams.SessionIdOverride;
-		NewSessionLANRef->SessionSettings = OpParams.SessionSettings;
+	TSharedRef<FSessionLAN> NewSessionLANRef = MakeShared<FSessionLAN>();
+	NewSessionLANRef->OwnerAccountId = Params.LocalAccountId;
+	NewSessionLANRef->SessionInfo.bAllowSanctionedPlayers = Params.bAllowSanctionedPlayers;
+	NewSessionLANRef->SessionInfo.bAntiCheatProtected = Params.bAntiCheatProtected;
+	NewSessionLANRef->SessionInfo.bIsDedicatedServerSession = IsRunningDedicatedServer();
+	NewSessionLANRef->SessionInfo.bIsLANSession = true;
+	NewSessionLANRef->SessionInfo.SessionId = FOnlineSessionIdRegistryLAN::GetChecked(Services.GetServicesProvider()).GetNextSessionId();
+	NewSessionLANRef->SessionInfo.SessionIdOverride = Params.SessionIdOverride;
+	NewSessionLANRef->SessionSettings = Params.SessionSettings;
 
-		// For LAN sessions, we'll add the Session member manually instead of calling JoinSession since there is no API calls involved
-		NewSessionLANRef->SessionMembers.Emplace(OpParams.LocalAccountId);
+	// For LAN sessions, we'll add the Session member manually instead of calling JoinSession since there is no API calls involved
+	NewSessionLANRef->SessionMembers.Emplace(Params.LocalAccountId);
 
-		// We save the local object for the session, and set up the appropriate references
-		AddSessionWithReferences(NewSessionLANRef, OpParams.SessionName, OpParams.LocalAccountId, OpParams.bPresenceEnabled);
+	// We save the local object for the session, and set up the appropriate references
+	AddSessionWithReferences(NewSessionLANRef, Params.SessionName, Params.LocalAccountId, Params.bPresenceEnabled);
 
-		Op.SetResult({ });
-	})
-	.Enqueue(GetSerialQueue());
+	Promise.EmplaceValue(FCreateSession::Result{ });
 
-	return Op->GetHandle();
+	return Future;
 }
 
-TFuture<TOnlineResult<FUpdateSessionSettingsImpl>> FSessionsLAN::UpdateSessionSettingsImpl(FUpdateSessionSettingsImpl::Params&& Params)
+TFuture<TOnlineResult<FUpdateSessionSettings>> FSessionsLAN::UpdateSessionSettingsImpl(const FUpdateSessionSettings::Params& Params)
 {
-	TPromise<TOnlineResult<FUpdateSessionSettingsImpl>> Promise;
-	TFuture<TOnlineResult<FUpdateSessionSettingsImpl>> Future = Promise.GetFuture();
+	TPromise<TOnlineResult<FUpdateSessionSettings>> Promise;
+	TFuture<TOnlineResult<FUpdateSessionSettings>> Future = Promise.GetFuture();
 
 	TOnlineResult<FGetMutableSessionByName> GetMutableSessionByNameResult = GetMutableSessionByName({ Params.SessionName });
 	check(GetMutableSessionByNameResult.IsOk());
@@ -217,7 +197,7 @@ TFuture<TOnlineResult<FUpdateSessionSettingsImpl>> FSessionsLAN::UpdateSessionSe
 	(*FoundSession) += SessionUpdateData;
 
 	// We set the result and fire the event
-	Promise.EmplaceValue(FUpdateSessionSettingsImpl::Result{ });
+	Promise.EmplaceValue(FUpdateSessionSettings::Result{ });
 
 	FSessionUpdated SessionUpdatedEvent{ Params.SessionName, SessionUpdateData };
 	SessionEvents.OnSessionUpdated.Broadcast(SessionUpdatedEvent);
@@ -225,128 +205,81 @@ TFuture<TOnlineResult<FUpdateSessionSettingsImpl>> FSessionsLAN::UpdateSessionSe
 	return Future;
 }
 
-TOnlineAsyncOpHandle<FFindSessions> FSessionsLAN::FindSessions(FFindSessions::Params&& Params)
+TFuture<TOnlineResult<FFindSessions>> FSessionsLAN::FindSessionsImpl(const FFindSessions::Params& Params)
 {
-	TOnlineAsyncOpRef<FFindSessions> Op = GetOp<FFindSessions>(MoveTemp(Params));
-	const FFindSessions::Params& OpParams = Op->GetParams();
+	TPromise<TOnlineResult<FFindSessions>> Promise;
+	TFuture<TOnlineResult<FFindSessions>> Future = Promise.GetFuture();
 
-	FOnlineError ParamsCheck = CheckFindSessionsParams(OpParams);
-	if (ParamsCheck != Errors::Success())
-	{
-		Op->SetError(MoveTemp(ParamsCheck));
-		return Op->GetHandle();
-	}
-
-	Op->Then([this](TOnlineAsyncOp<FFindSessions>& Op) mutable
-	{
-		const FFindSessions::Params& OpParams = Op.GetParams();
+	// Before we start the search, we reset the cache and save the promise
+	SearchResultsUserMap.FindOrAdd(Params.LocalAccountId).Reset();
+	CurrentSessionSearchPromisesUserMap.Emplace(Params.LocalAccountId, MoveTemp(Promise));
 
-		FOnlineError StateCheck = CheckFindSessionsState(OpParams);
-		if (StateCheck != Errors::Success())
-		{
-			Op.SetError(MoveTemp(StateCheck));
-			return;
-		}
+	FindLANSessions(Params.LocalAccountId);
 
-		CurrentSessionSearchHandlesUserMap.Emplace(OpParams.LocalAccountId, Op.AsShared());
-
-		FindLANSessions(OpParams.LocalAccountId);
-	})
-	.Enqueue(GetSerialQueue());
-
-	return Op->GetHandle();
+	return Future;
 }
 
-TOnlineAsyncOpHandle<FJoinSession> FSessionsLAN::JoinSession(FJoinSession::Params&& Params)
+TFuture<TOnlineResult<FJoinSession>> FSessionsLAN::JoinSessionImpl(const FJoinSession::Params& Params)
 {
-	TOnlineAsyncOpRef<FJoinSession> Op = GetOp<FJoinSession>(MoveTemp(Params));
-	const FJoinSession::Params& OpParams = Op->GetParams();
-
-	FOnlineError ParamsCheck = CheckJoinSessionParams(OpParams);
-	if (ParamsCheck != Errors::Success())
-	{
-		Op->SetError(MoveTemp(ParamsCheck));
-		return Op->GetHandle();
-	}
+	TPromise<TOnlineResult<FJoinSession>> Promise;
+	TFuture<TOnlineResult<FJoinSession>> Future = Promise.GetFuture();
 
-	// If no restrictions apply, we'll save our copy of the session, add the new players, and register it
-	Op->Then([this](TOnlineAsyncOp<FJoinSession>& Op) mutable
+	TOnlineResult<FGetSessionById> GetSessionByIdResult = GetSessionById({ Params.SessionId });
+	if (GetSessionByIdResult.IsError())
 	{
-		const FJoinSession::Params& OpParams = Op.GetParams();
-
-		FOnlineError StateCheck = CheckJoinSessionState(OpParams);
-		if (StateCheck != Errors::Success())
+		// If no result is found, the id might be expired, which we should notify
+		if (FOnlineSessionIdRegistryLAN::GetChecked(Services.GetServicesProvider()).IsSessionIdExpired(Params.SessionId))
 		{
-			Op.SetError(MoveTemp(StateCheck));
-			return;
+			UE_LOG(LogTemp, Warning, TEXT("[%s] SessionId parameter [%s] is expired. Please call FindSessions to get an updated list of available sessions "), UTF8_TO_TCHAR(__FUNCTION__), *ToLogString(Params.SessionId));
 		}
 
-		TOnlineResult<FGetSessionById> GetSessionByIdResult = GetSessionById({ OpParams.SessionId });
-		if (GetSessionByIdResult.IsError())
-		{
-			// If no result is found, the id might be expired, which we should notify
-			if (FOnlineSessionIdRegistryLAN::GetChecked(Services.GetServicesProvider()).IsSessionIdExpired(OpParams.SessionId))
-			{
-				UE_LOG(LogTemp, Warning, TEXT("[FSessionsLAN::JoinSession] SessionId parameter [%s] is expired. Please call FindSessions to get an updated list of available sessions "), *ToLogString(OpParams.SessionId));
-			}
-
-			Op.SetError(MoveTemp(GetSessionByIdResult.GetErrorValue()));
-			return ;
-		}
+		Promise.EmplaceValue(GetSessionByIdResult.GetErrorValue());
+		return Future;
+	}
 
-		const TSharedRef<const ISession>& FoundSession = GetSessionByIdResult.GetOkValue().Session;
+	const TSharedRef<const ISession>& FoundSession = GetSessionByIdResult.GetOkValue().Session;
 
-		// We set up the appropriate references for the session
-		AddSessionReferences(FoundSession->GetSessionId(), OpParams.SessionName, OpParams.LocalAccountId, OpParams.bPresenceEnabled);
+	// We set up the appropriate references for the session
+	AddSessionReferences(FoundSession->GetSessionId(), Params.SessionName, Params.LocalAccountId, Params.bPresenceEnabled);
 
-		Op.SetResult(FJoinSession::Result{ });
+	Promise.EmplaceValue(FJoinSession::Result{ });
 
-		FSessionJoined SessionJoinedEvent = { OpParams.LocalAccountId, FoundSession->GetSessionId() };
+	FSessionJoined SessionJoinedEvent = { Params.LocalAccountId, FoundSession->GetSessionId() };
 		
-		SessionEvents.OnSessionJoined.Broadcast(SessionJoinedEvent);
-	})
-	.Enqueue(GetSerialQueue());
+	SessionEvents.OnSessionJoined.Broadcast(SessionJoinedEvent);
 
-	return Op->GetHandle();
+	return Future;
 }
 
-TOnlineAsyncOpHandle<FLeaveSession> FSessionsLAN::LeaveSession(FLeaveSession::Params&& Params)
+TFuture<TOnlineResult<FLeaveSession>> FSessionsLAN::LeaveSessionImpl(const FLeaveSession::Params& Params)
 {
-	TOnlineAsyncOpRef<FLeaveSession> Op = GetOp<FLeaveSession>(MoveTemp(Params));
+	TPromise<TOnlineResult<FLeaveSession>> Promise;
+	TFuture<TOnlineResult<FLeaveSession>> Future = Promise.GetFuture();
 
-	Op->Then([this](TOnlineAsyncOp<FLeaveSession>& Op) mutable
+	TOnlineResult<FGetSessionByName> GetSessionByNameResult = GetSessionByName({ Params.SessionName });
+	if (GetSessionByNameResult.IsOk())
 	{
-		const FLeaveSession::Params& OpParams = Op.GetParams();
+		TSharedRef<const ISession> FoundSession = GetSessionByNameResult.GetOkValue().Session;
 
-		FOnlineError StateCheck = CheckLeaveSessionState(OpParams);
-		if (StateCheck != Errors::Success())
+		if (FoundSession->GetOwnerAccountId() == Params.LocalAccountId && FoundSession->GetSessionSettings().JoinPolicy == ESessionJoinPolicy::Public)
 		{
-			Op.SetError(MoveTemp(StateCheck));
-			return;
+			StopLANSession();
 		}
 
-		TOnlineResult<FGetSessionByName> GetSessionByNameResult = GetSessionByName({ OpParams.SessionName });
-		if (GetSessionByNameResult.IsOk())
-		{
-			TSharedRef<const ISession> FoundSession = GetSessionByNameResult.GetOkValue().Session;
+		ClearSessionReferences(FoundSession->GetSessionId(), Params.SessionName, Params.LocalAccountId);
 
-			if (FoundSession->GetOwnerAccountId() == OpParams.LocalAccountId && FoundSession->GetSessionSettings().JoinPolicy == ESessionJoinPolicy::Public)
-			{
-				StopLANSession();
-			}
-
-			ClearSessionReferences(FoundSession->GetSessionId(), OpParams.SessionName, OpParams.LocalAccountId);
-		}
-
-		Op.SetResult(FLeaveSession::Result{ });
+		Promise.EmplaceValue(FLeaveSession::Result{ });
 
 		FSessionLeft SessionLeftEvent;
-		SessionLeftEvent.LocalAccountId = OpParams.LocalAccountId;
+		SessionLeftEvent.LocalAccountId = Params.LocalAccountId;
 		SessionEvents.OnSessionLeft.Broadcast(SessionLeftEvent);
-	})
-	.Enqueue(GetSerialQueue());
+	}
+	else
+	{
+		Promise.EmplaceValue(GetSessionByNameResult.GetErrorValue());
+	}
 
-	return Op->GetHandle();
+	return Future;
 }
 
 /** LANSessionManager methods */
@@ -406,9 +339,8 @@ void FSessionsLAN::FindLANSessions(const FAccountId& LocalAccountId)
 		}
 
 		// Trigger the delegate as having failed
-		CurrentSessionSearchHandlesUserMap.FindChecked(LocalAccountId)->SetError(Errors::RequestFailure()); // TODO: May need a new error type
-
-		CurrentSessionSearchHandlesUserMap.Remove(LocalAccountId);
+		CurrentSessionSearchPromisesUserMap.FindChecked(LocalAccountId).EmplaceValue(Errors::RequestFailure()); // TODO: May need a new error type
+		CurrentSessionSearchPromisesUserMap.Remove(LocalAccountId);
 
 		// If we were hosting public sessions before the search, we'll return the beacon to that state
 		if (PublicSessionsHosted > 0)
@@ -484,9 +416,8 @@ void FSessionsLAN::OnLANSearchTimeout(const FAccountId LocalAccountId)
 
 	TArray<FOnlineSessionId>& FoundSessionIds = SearchResultsUserMap.FindChecked(LocalAccountId);
 
-	CurrentSessionSearchHandlesUserMap.FindChecked(LocalAccountId)->SetResult({ FoundSessionIds });
-
-	CurrentSessionSearchHandlesUserMap.Remove(LocalAccountId);
+	CurrentSessionSearchPromisesUserMap.FindChecked(LocalAccountId).EmplaceValue(FFindSessions::Result{ FoundSessionIds });
+	CurrentSessionSearchPromisesUserMap.Remove(LocalAccountId); 
 
 	// If we were hosting public sessions before the search, we'll return the beacon to that state
 	if (PublicSessionsHosted > 0)
diff --git a/Engine/Plugins/Online/Experimental/OnlineServices/Source/OnlineServicesCommonEngineUtils/Public/Online/SessionsLAN.h b/Engine/Plugins/Online/Experimental/OnlineServices/Source/OnlineServicesCommonEngineUtils/Public/Online/SessionsLAN.h
--- a/Engine/Plugins/Online/Experimental/OnlineServices/Source/OnlineServicesCommonEngineUtils/Public/Online/SessionsLAN.h
+++ b/Engine/Plugins/Online/Experimental/OnlineServices/Source/OnlineServicesCommonEngineUtils/Public/Online/SessionsLAN.h
@@ -57,14 +57,12 @@ public:
 
 	virtual void Tick(float DeltaSeconds) override;
 
-	// ISessions
-	virtual TOnlineAsyncOpHandle<FCreateSession> CreateSession(FCreateSession::Params&& Params) override;
-	virtual TOnlineAsyncOpHandle<FFindSessions> FindSessions(FFindSessions::Params&& Params) override;
-	virtual TOnlineAsyncOpHandle<FJoinSession> JoinSession(FJoinSession::Params&& Params) override;
-	virtual TOnlineAsyncOpHandle<FLeaveSession> LeaveSession(FLeaveSession::Params&& Params) override;
-
 	// FSessionsCommon
-	virtual TFuture<TOnlineResult<FUpdateSessionSettingsImpl>> UpdateSessionSettingsImpl(FUpdateSessionSettingsImpl::Params&& Params) override;
+	virtual TFuture<TOnlineResult<FCreateSession>> CreateSessionImpl(const FCreateSession::Params& Params) override;
+	virtual TFuture<TOnlineResult<FUpdateSessionSettings>> UpdateSessionSettingsImpl(const FUpdateSessionSettings::Params& Params) override;
+	virtual TFuture<TOnlineResult<FLeaveSession>> LeaveSessionImpl(const FLeaveSession::Params& Params) override;
+	virtual TFuture<TOnlineResult<FFindSessions>> FindSessionsImpl(const FFindSessions::Params& Params) override;
+	virtual TFuture<TOnlineResult<FJoinSession>> JoinSessionImpl(const FJoinSession::Params& Params) override;
 
 protected:
 	TOptional<FOnlineError> TryHostLANSession();
diff --git a/Engine/Plugins/Online/Experimental/OnlineServicesEOSGS/Source/Private/Online/SessionsEOSGS.cpp b/Engine/Plugins/Online/Experimental/OnlineServicesEOSGS/Source/Private/Online/SessionsEOSGS.cpp
--- a/Engine/Plugins/Online/Experimental/OnlineServicesEOSGS/Source/Private/Online/SessionsEOSGS.cpp
+++ b/Engine/Plugins/Online/Experimental/OnlineServicesEOSGS/Source/Private/Online/SessionsEOSGS.cpp
@@ -314,141 +314,118 @@ void FSessionsEOSGS::HandleJoinSessionAccepted(const EOS_Sessions_JoinSessionAcc
 	}
 }
 
-TOnlineAsyncOpHandle<FCreateSession> FSessionsEOSGS::CreateSession(FCreateSession::Params&& Params)
+TFuture<TOnlineResult<FCreateSession>> FSessionsEOSGS::CreateSessionImpl(const FCreateSession::Params& Params)
 {
 	// LAN Sessions
 	if (Params.bIsLANSession)
 	{
-		return FSessionsLAN::CreateSession(MoveTemp(Params));
+		return FSessionsLAN::CreateSessionImpl(Params);
 	}
 
 	// EOSGS Sessions
 
-	TOnlineAsyncOpRef<FCreateSession> Op = GetOp<FCreateSession>(MoveTemp(Params));
-	const FCreateSession::Params& OpParams = Op->GetParams();
+	TPromise<TOnlineResult<FCreateSession>> Promise;
+	TFuture<TOnlineResult<FCreateSession>> Future = Promise.GetFuture();
 
-	FOnlineError ParamsCheck = CheckCreateSessionParams(OpParams);
-	if (ParamsCheck != Errors::Success())
+	if (!Params.SessionIdOverride.IsEmpty())
 	{
-		Op->SetError(MoveTemp(ParamsCheck));
-		return Op->GetHandle();
-	}
-
-	Op->Then([this](TOnlineAsyncOp<FCreateSession>& Op)
-	{
-		const FCreateSession::Params& OpParams = Op.GetParams();
+		int32 Length = Params.SessionIdOverride.Len();
 
-		FOnlineError StateCheck = CheckCreateSessionState(OpParams);
-		if (StateCheck != Errors::Success())
+		if (Length < EOS_SESSIONMODIFICATION_MIN_SESSIONIDOVERRIDE_LENGTH || Length > EOS_SESSIONMODIFICATION_MAX_SESSIONIDOVERRIDE_LENGTH)
 		{
-			Op.SetError(MoveTemp(StateCheck));
-
-			return MakeFulfilledPromise<TDefaultErrorResult<FUpdateSessionImplEOSGS>>().GetFuture();
-		}
-
-		if (!OpParams.SessionIdOverride.IsEmpty())
-		{
-			int32 Length = OpParams.SessionIdOverride.Len();
-
-			if (Length < EOS_SESSIONMODIFICATION_MIN_SESSIONIDOVERRIDE_LENGTH || Length > EOS_SESSIONMODIFICATION_MAX_SESSIONIDOVERRIDE_LENGTH)
-			{
-				UE_LOG(LogTemp, Warning, TEXT("[FSessionsEOSGS::CreateSession] Could not create session with SessionIdOverride [%s] of size [%d]. SessionIdOverride size must be between [%d] and [%d] characters long"), *OpParams.SessionIdOverride, Length, EOS_SESSIONMODIFICATION_MIN_SESSIONIDOVERRIDE_LENGTH, EOS_SESSIONMODIFICATION_MAX_SESSIONIDOVERRIDE_LENGTH);
+			UE_LOG(LogTemp, Warning, TEXT("[FSessionsEOSGS::CreateSession] Could not create session with SessionIdOverride [%s] of size [%d]. SessionIdOverride size must be between [%d] and [%d] characters long"), *Params.SessionIdOverride, Length, EOS_SESSIONMODIFICATION_MIN_SESSIONIDOVERRIDE_LENGTH, EOS_SESSIONMODIFICATION_MAX_SESSIONIDOVERRIDE_LENGTH);
 
-				Op.SetError(Errors::InvalidParams());
+			Promise.EmplaceValue(Errors::InvalidParams());
 
-				return MakeFulfilledPromise<TDefaultErrorResult<FUpdateSessionImplEOSGS>>().GetFuture();
-			}
+			return Future;
 		}
+	}
 
-		// Check if the Bucket Id custom setting is set. EOS Sessions can not be created without it
-		if (!OpParams.SessionSettings.CustomSettings.Contains(EOSGS_BUCKET_ID_ATTRIBUTE_KEY))
-		{
-			UE_LOG(LogTemp, Warning, TEXT("[FSessionsEOSGS::CreateSession] Could not create session without Custom Setting 'EOSGS_BUCKET_ID_ATTRIBUTE_KEY' (FString) set."));
+	// Check if the Bucket Id custom setting is set. EOS Sessions can not be created without it
+	if (!Params.SessionSettings.CustomSettings.Contains(EOSGS_BUCKET_ID_ATTRIBUTE_KEY))
+	{
+		UE_LOG(LogTemp, Warning, TEXT("[FSessionsEOSGS::CreateSession] Could not create session without Custom Setting 'EOSGS_BUCKET_ID_ATTRIBUTE_KEY' (FString) set."));
 
-			Op.SetError(Errors::InvalidParams());
+		Promise.EmplaceValue(Errors::InvalidParams());
 
-			return MakeFulfilledPromise<TDefaultErrorResult<FUpdateSessionImplEOSGS>>().GetFuture();
-		}
+		return Future;
+	}
 
-		// After all initial checks, we start the session creation operations
+	// After all initial checks, we start the session creation operations
 
-		EOS_Sessions_CreateSessionModificationOptions CreateSessionModificationOptions = {};
-		CreateSessionModificationOptions.ApiVersion = EOS_SESSIONS_CREATESESSIONMODIFICATION_API_LATEST;
-		static_assert(EOS_SESSIONS_CREATESESSIONMODIFICATION_API_LATEST == 4, "EOS_Sessions_CreateSessionModificationOptions updated, check new fields");
+	EOS_Sessions_CreateSessionModificationOptions CreateSessionModificationOptions = {};
+	CreateSessionModificationOptions.ApiVersion = EOS_SESSIONS_CREATESESSIONMODIFICATION_API_LATEST;
+	static_assert(EOS_SESSIONS_CREATESESSIONMODIFICATION_API_LATEST == 4, "EOS_Sessions_CreateSessionModificationOptions updated, check new fields");
 
-		CreateSessionModificationOptions.bPresenceEnabled = OpParams.bPresenceEnabled;
-		CreateSessionModificationOptions.bSanctionsEnabled = !OpParams.bAllowSanctionedPlayers;
+	CreateSessionModificationOptions.bPresenceEnabled = Params.bPresenceEnabled;
+	CreateSessionModificationOptions.bSanctionsEnabled = !Params.bAllowSanctionedPlayers;
 
-		const FCustomSessionSetting* BucketIdSetting = OpParams.SessionSettings.CustomSettings.Find(EOSGS_BUCKET_ID_ATTRIBUTE_KEY);
-		const FTCHARToUTF8 BucketIdUtf8(BucketIdSetting ? *BucketIdSetting->Data.GetString() : nullptr);
+	const FCustomSessionSetting* BucketIdSetting = Params.SessionSettings.CustomSettings.Find(EOSGS_BUCKET_ID_ATTRIBUTE_KEY);
+	const FTCHARToUTF8 BucketIdUtf8(BucketIdSetting ? *BucketIdSetting->Data.GetString() : nullptr);
 
-		if (BucketIdUtf8.Length())
-		{
-			CreateSessionModificationOptions.BucketId = BucketIdUtf8.Get();
-		}
+	if (BucketIdUtf8.Length())
+	{
+		CreateSessionModificationOptions.BucketId = BucketIdUtf8.Get();
+	}
 	
-		CreateSessionModificationOptions.LocalUserId = GetProductUserIdChecked(OpParams.LocalAccountId);
-		CreateSessionModificationOptions.MaxPlayers = OpParams.SessionSettings.NumMaxConnections;
+	CreateSessionModificationOptions.LocalUserId = GetProductUserIdChecked(Params.LocalAccountId);
+	CreateSessionModificationOptions.MaxPlayers = Params.SessionSettings.NumMaxConnections;
 
-		const FTCHARToUTF8 SessionIdUtf8(*OpParams.SessionIdOverride);
-		if (SessionIdUtf8.Length())
-		{
-			CreateSessionModificationOptions.SessionId = SessionIdUtf8.Get();
-		}
+	const FTCHARToUTF8 SessionIdUtf8(*Params.SessionIdOverride);
+	if (SessionIdUtf8.Length())
+	{
+		CreateSessionModificationOptions.SessionId = SessionIdUtf8.Get();
+	}
 
-		const FTCHARToUTF8 SessionNameUtf8(*OpParams.SessionName.ToString());
-		CreateSessionModificationOptions.SessionName = SessionNameUtf8.Get();
+	const FTCHARToUTF8 SessionNameUtf8(*Params.SessionName.ToString());
+	CreateSessionModificationOptions.SessionName = SessionNameUtf8.Get();
 
-		EOS_HSessionModification SessionModificationHandle = nullptr;
-		EOS_EResult ResultCode = EOS_Sessions_CreateSessionModification(SessionsHandle, &CreateSessionModificationOptions, &SessionModificationHandle);
-		if (ResultCode != EOS_EResult::EOS_Success)
-		{
-			UE_LOG(LogTemp, Error, TEXT("[FSessionsEOSGS::CreateSession] EOS_Sessions_CreateSessionModification failed with result [%s]"), *LexToString(ResultCode));
+	EOS_HSessionModification SessionModificationHandle = nullptr;
+	EOS_EResult ResultCode = EOS_Sessions_CreateSessionModification(SessionsHandle, &CreateSessionModificationOptions, &SessionModificationHandle);
+	if (ResultCode != EOS_EResult::EOS_Success)
+	{
+		UE_LOG(LogTemp, Error, TEXT("[FSessionsEOSGS::CreateSession] EOS_Sessions_CreateSessionModification failed with result [%s]"), *LexToString(ResultCode));
 
-			Op.SetError(Errors::FromEOSResult(ResultCode));
+		Promise.EmplaceValue(Errors::FromEOSResult(ResultCode));
 
-			return MakeFulfilledPromise<TDefaultErrorResult<FUpdateSessionImplEOSGS>>().GetFuture();
-		}
+		return Future;
+	}
 
-		// TODO: We could call EOS_SessionModification_SetHostAddress at this point, although it's not necessary
+	// TODO: We could call EOS_SessionModification_SetHostAddress at this point, although it's not necessary
 
-		// We write all SessionSettings values into the SessionModificationHandle
-		WriteCreateSessionModificationHandle(SessionModificationHandle, OpParams);
+	// We write all SessionSettings values into the SessionModificationHandle
+	WriteCreateSessionModificationHandle(SessionModificationHandle, Params);
 
-		// Always update joinability on session creation
-		FUpdateSessionImplEOSGS::Params UpdateSessionImplParams{ MakeShared<FSessionModificationHandleEOSGS>(SessionModificationHandle), FUpdateSessionJoinabilityParams{ OpParams.SessionName, OpParams.SessionSettings.bAllowNewMembers } };
+	// Always update joinability on session creation
+	FUpdateSessionImplEOSGS::Params UpdateSessionImplParams{ MakeShared<FSessionModificationHandleEOSGS>(SessionModificationHandle), FUpdateSessionJoinabilityParams{ Params.SessionName, Params.SessionSettings.bAllowNewMembers } };
 		
-		return UpdateSessionImplEOSGS(MoveTemp(UpdateSessionImplParams));
-	})
-	.Then([this](TOnlineAsyncOp<FCreateSession>& Op, TDefaultErrorResult<FUpdateSessionImplEOSGS>&& Result)
-	{
-		const FCreateSession::Params& OpParams = Op.GetParams();
-
-		if (Result.IsOk())
-		{
-			// If the session is created successfully, we'll create the local session object
-			TSharedRef<FSessionEOSGS> NewSessionEOSGSRef = MakeShared<FSessionEOSGS>();
-			NewSessionEOSGSRef->OwnerAccountId = OpParams.LocalAccountId;
-			NewSessionEOSGSRef->SessionSettings = OpParams.SessionSettings;
-			NewSessionEOSGSRef->SessionInfo.SessionId = CreateSessionId(Result.GetOkValue().NewSessionId);
-			NewSessionEOSGSRef->SessionInfo.bAllowSanctionedPlayers = OpParams.bAllowSanctionedPlayers;
-			NewSessionEOSGSRef->SessionInfo.bAntiCheatProtected = OpParams.bAntiCheatProtected;
-			NewSessionEOSGSRef->SessionInfo.bIsDedicatedServerSession = IsRunningDedicatedServer();
-			NewSessionEOSGSRef->SessionInfo.bIsLANSession = false;
-			NewSessionEOSGSRef->SessionInfo.SessionIdOverride = OpParams.SessionIdOverride;
-
-			AddSessionWithReferences(NewSessionEOSGSRef, OpParams.SessionName, OpParams.LocalAccountId, OpParams.bPresenceEnabled);
-
-			Op.SetResult({ });
-		}
-		else
-		{
-			Op.SetError(MoveTemp(Result.GetErrorValue()));
-		}
-	})
-	.Enqueue(GetSerialQueue());
+	UpdateSessionImplEOSGS(MoveTemp(UpdateSessionImplParams))
+		.Next([this, Promise = MoveTemp(Promise), Params](TDefaultErrorResult<FUpdateSessionImplEOSGS>&& Result) mutable
+			{
+				if (Result.IsOk())
+				{
+					// If the session is created successfully, we'll create the local session object
+					TSharedRef<FSessionEOSGS> NewSessionEOSGSRef = MakeShared<FSessionEOSGS>();
+					NewSessionEOSGSRef->OwnerAccountId = Params.LocalAccountId;
+					NewSessionEOSGSRef->SessionSettings = Params.SessionSettings;
+					NewSessionEOSGSRef->SessionInfo.SessionId = CreateSessionId(Result.GetOkValue().NewSessionId);
+					NewSessionEOSGSRef->SessionInfo.bAllowSanctionedPlayers = Params.bAllowSanctionedPlayers;
+					NewSessionEOSGSRef->SessionInfo.bAntiCheatProtected = Params.bAntiCheatProtected;
+					NewSessionEOSGSRef->SessionInfo.bIsDedicatedServerSession = IsRunningDedicatedServer();
+					NewSessionEOSGSRef->SessionInfo.bIsLANSession = false;
+					NewSessionEOSGSRef->SessionInfo.SessionIdOverride = Params.SessionIdOverride;
+
+					AddSessionWithReferences(NewSessionEOSGSRef, Params.SessionName, Params.LocalAccountId, Params.bPresenceEnabled);
+
+					Promise.EmplaceValue(FCreateSession::Result({ }));
+				}
+				else
+				{
+					Promise.EmplaceValue(Result.GetErrorValue());
+				}
+			});
 
-	return Op->GetHandle();
+	return Future;
 }
 
 void FSessionsEOSGS::SetJoinInProgressAllowed(EOS_HSessionModification& SessionModHandle, bool bIsJoinInProgressAllowed)
@@ -649,19 +626,19 @@ void FSessionsEOSGS::WriteUpdateSessionModificationHandle(EOS_HSessionModificati
 	}
 }
 
-TFuture<TOnlineResult<FUpdateSessionSettingsImpl>> FSessionsEOSGS::UpdateSessionSettingsImpl(FUpdateSessionSettingsImpl::Params&& Params)
+TFuture<TOnlineResult<FUpdateSessionSettings>> FSessionsEOSGS::UpdateSessionSettingsImpl(const FUpdateSessionSettings::Params& Params)
 {
 	// LAN Sessions
 	TOnlineResult<FGetSessionByName> Result = GetSessionByName({ Params.SessionName });
 	if (Result.IsOk() && Result.GetOkValue().Session->GetSessionInfo().bIsLANSession)
 	{
-		return FSessionsLAN::UpdateSessionSettingsImpl(MoveTemp(Params));
+		return FSessionsLAN::UpdateSessionSettingsImpl(Params);
 	}
 
 	// EOSGS Sessions
 
-	TPromise<TOnlineResult<FUpdateSessionSettingsImpl>> Promise;
-	TFuture<TOnlineResult<FUpdateSessionSettingsImpl>> Future = Promise.GetFuture();
+	TPromise<TOnlineResult<FUpdateSessionSettings>> Promise;
+	TFuture<TOnlineResult<FUpdateSessionSettings>> Future = Promise.GetFuture();
 
 	EOS_Sessions_UpdateSessionModificationOptions UpdateSessionModificationOptions = {};
 	UpdateSessionModificationOptions.ApiVersion = EOS_SESSIONS_UPDATESESSIONMODIFICATION_API_LATEST;
@@ -694,24 +671,24 @@ TFuture<TOnlineResult<FUpdateSessionSettingsImpl>> FSessionsEOSGS::UpdateSession
 	};
 
 	UpdateSessionImplEOSGS(MoveTemp(UpdateSessionImplEOSGSParams))
-	.Next([this, Promise = MoveTemp(Promise), OpParams = MoveTemp(Params)](TDefaultErrorResult<FUpdateSessionImplEOSGS>&& Result) mutable
+	.Next([this, Promise = MoveTemp(Promise), Params](TDefaultErrorResult<FUpdateSessionImplEOSGS>&& Result) mutable
 	{
 		if (Result.IsOk())
 		{
-			TOnlineResult<FGetMutableSessionByName> GetMutableSessionByNameResult = GetMutableSessionByName({ OpParams.SessionName });
+			TOnlineResult<FGetMutableSessionByName> GetMutableSessionByNameResult = GetMutableSessionByName({ Params.SessionName });
 			if (GetMutableSessionByNameResult.IsOk())
 			{
 				TSharedRef<FSessionCommon> FoundSession = GetMutableSessionByNameResult.GetOkValue().Session;
 
 				// Now that the API Session update has processed successfully, we'll update our local session with the same data
-				FSessionUpdate SessionUpdateData = BuildSessionUpdate(FoundSession, OpParams.Mutations);
+				FSessionUpdate SessionUpdateData = BuildSessionUpdate(FoundSession, Params.Mutations);
 
 				(*FoundSession) += SessionUpdateData;
 
 				// We set the result and fire the event
-				Promise.EmplaceValue(FUpdateSessionSettingsImpl::Result{ });
+				Promise.EmplaceValue(FUpdateSessionSettings::Result{ });
 
-				FSessionUpdated SessionUpdatedEvent{ OpParams.SessionName, SessionUpdateData };
+				FSessionUpdated SessionUpdatedEvent{ Params.SessionName, SessionUpdateData };
 				SessionEvents.OnSessionUpdated.Broadcast(SessionUpdatedEvent);
 			}
 			else
@@ -905,67 +882,56 @@ TOnlineAsyncOpHandle<FSendSingleSessionInviteImpl> FSessionsEOSGS::SendSingleSes
 	return Op->GetHandle();
 }
 
-TOnlineAsyncOpHandle<FLeaveSession> FSessionsEOSGS::LeaveSession(FLeaveSession::Params&& Params)
+TFuture<TOnlineResult<FLeaveSession>> FSessionsEOSGS::LeaveSessionImpl(const FLeaveSession::Params& Params)
 {
 	// LAN Sessions
 	TOnlineResult<FGetSessionByName> Result = GetSessionByName({ Params.SessionName });
 	if (Result.IsOk() && Result.GetOkValue().Session->GetSessionInfo().bIsLANSession)
 	{
-		return FSessionsLAN::LeaveSession(MoveTemp(Params));
+		return FSessionsLAN::LeaveSessionImpl(Params);
 	}
 
 	// EOSGS Sessions
 
-	TOnlineAsyncOpRef<FLeaveSession> Op = GetOp<FLeaveSession>(MoveTemp(Params));
+	TPromise<TOnlineResult<FLeaveSession>> Promise;
+	TFuture<TOnlineResult<FLeaveSession>> Future = Promise.GetFuture();
 
-	Op->Then([this](TOnlineAsyncOp<FLeaveSession>& Op, TPromise<const EOS_Sessions_DestroySessionCallbackInfo*>&& Promise)
-	{
-		const FLeaveSession::Params& OpParams = Op.GetParams();
+	EOS_Sessions_DestroySessionOptions DestroySessionOptions = {};
+	DestroySessionOptions.ApiVersion = EOS_SESSIONS_DESTROYSESSION_API_LATEST;
+	static_assert(EOS_SESSIONS_DESTROYSESSION_API_LATEST == 1, "EOS_Sessions_DestroySessionOptions updated, check new fields");
 
-		FOnlineError StateCheck = CheckLeaveSessionState(OpParams);
-		if (StateCheck != Errors::Success())
-		{
-			Op.SetError(MoveTemp(StateCheck));
-			Promise.EmplaceValue();
-			return;
-		}
-
-		EOS_Sessions_DestroySessionOptions DestroySessionOptions = {};
-		DestroySessionOptions.ApiVersion = EOS_SESSIONS_DESTROYSESSION_API_LATEST;
-		static_assert(EOS_SESSIONS_DESTROYSESSION_API_LATEST == 1, "EOS_Sessions_DestroySessionOptions updated, check new fields");
-
-		const FTCHARToUTF8 SessionNameUtf8(*OpParams.SessionName.ToString());
-		DestroySessionOptions.SessionName = SessionNameUtf8.Get();
+	const FTCHARToUTF8 SessionNameUtf8(*Params.SessionName.ToString());
+	DestroySessionOptions.SessionName = SessionNameUtf8.Get();
 
-		EOS_Async(EOS_Sessions_DestroySession, SessionsHandle, DestroySessionOptions, MoveTemp(Promise));
-	})
-	.Then([this](TOnlineAsyncOp<FLeaveSession>& Op, const EOS_Sessions_DestroySessionCallbackInfo* Result)
-	{
-		if (Result->ResultCode != EOS_EResult::EOS_Success)
+	EOS_Async(EOS_Sessions_DestroySession, SessionsHandle, DestroySessionOptions,
+		[this, Promise = MoveTemp(Promise), Params](const EOS_Sessions_DestroySessionCallbackInfo* Result) mutable
 		{
-			Op.SetError(Errors::FromEOSResult(Result->ResultCode));
-			return;
-		}
-
-		const FLeaveSession::Params& OpParams = Op.GetParams();
+			if (Result->ResultCode != EOS_EResult::EOS_Success)
+			{
+				Promise.EmplaceValue(Errors::FromEOSResult(Result->ResultCode));
+				return;
+			}
 
-		TOnlineResult<FGetSessionByName> GetSessionByNameResult = GetSessionByName({ OpParams.SessionName });
-		if (GetSessionByNameResult.IsOk())
-		{
-			TSharedRef<const ISession> FoundSession = GetSessionByNameResult.GetOkValue().Session;
+			TOnlineResult<FGetSessionByName> GetSessionByNameResult = GetSessionByName({ Params.SessionName });
+			if (GetSessionByNameResult.IsOk())
+			{
+				TSharedRef<const ISession> FoundSession = GetSessionByNameResult.GetOkValue().Session;
 
-			ClearSessionReferences(FoundSession->GetSessionId(), OpParams.SessionName, OpParams.LocalAccountId);
-		}
+				ClearSessionReferences(FoundSession->GetSessionId(), Params.SessionName, Params.LocalAccountId);
 
-		Op.SetResult(FLeaveSession::Result{ });
+				Promise.EmplaceValue(FLeaveSession::Result{ });
 
-		FSessionLeft SessionLeftEvent;
-		SessionLeftEvent.LocalAccountId = OpParams.LocalAccountId;
-		SessionEvents.OnSessionLeft.Broadcast(SessionLeftEvent);
-	})
-	.Enqueue(GetSerialQueue());
+				FSessionLeft SessionLeftEvent;
+				SessionLeftEvent.LocalAccountId = Params.LocalAccountId;
+				SessionEvents.OnSessionLeft.Broadcast(SessionLeftEvent);
+			}
+			else
+			{
+				Promise.EmplaceValue(GetSessionByNameResult.GetErrorValue());
+			}
+		});
 
-	return Op->GetHandle();
+	return Future;
 }
 
 void FSessionsEOSGS::SetSessionSearchMaxResults(FSessionSearchHandleEOSGS& SessionSearchHandle, uint32 MaxResults)
@@ -1065,182 +1031,150 @@ FSessionInviteId FSessionsEOSGS::CreateSessionInviteId(const FString& SessionInv
 	return FOnlineSessionInviteIdRegistryEOSGS::Get().BasicRegistry.FindOrAddHandle(SessionInviteId);
 }
 
-TOnlineAsyncOpHandle<FFindSessions> FSessionsEOSGS::FindSessions(FFindSessions::Params&& Params)
+TFuture<TOnlineResult<FFindSessions>> FSessionsEOSGS::FindSessionsImpl(const FFindSessions::Params& Params)
 {
 	// LAN Sessions
 	if (Params.bFindLANSessions)
 	{
-		return FSessionsLAN::FindSessions(MoveTemp(Params));
+		return FSessionsLAN::FindSessionsImpl(Params);
 	}
 
 	// EOSGS Sessions
+	TPromise<TOnlineResult<FFindSessions>> Promise;
+	TFuture<TOnlineResult<FFindSessions>> Future = Promise.GetFuture();
 
-	TOnlineAsyncOpRef<FFindSessions> Op = GetOp<FFindSessions>(MoveTemp(Params));
-	const FFindSessions::Params& OpParams = Op->GetParams();
+	// We start preparing the search
+	EOS_Sessions_CreateSessionSearchOptions CreateSessionSearchOptions = { };
+	CreateSessionSearchOptions.ApiVersion = EOS_SESSIONS_CREATESESSIONSEARCH_API_LATEST;
+	static_assert(EOS_SESSIONS_CREATESESSIONSEARCH_API_LATEST == 1, "EOS_Sessions_CreateSessionSearchOptions updated, check new fields");
+	CreateSessionSearchOptions.MaxSearchResults = FMath::Clamp(Params.MaxResults, 0, EOS_SESSIONS_MAX_SEARCH_RESULTS);
 
-	FOnlineError ParamsCheck = CheckFindSessionsParams(OpParams);
-	if (ParamsCheck != Errors::Success())
-	{
-		Op->SetError(MoveTemp(ParamsCheck));
-		return Op->GetHandle();
-	}
-
-	// We start the find operation
-	Op->Then([this](TOnlineAsyncOp<FFindSessions>& Op, TPromise<const EOS_SessionSearch_FindCallbackInfo*>&& Promise)
+	EOS_HSessionSearch SearchHandle = nullptr;
+	EOS_EResult ResultCode = EOS_Sessions_CreateSessionSearch(SessionsHandle, &CreateSessionSearchOptions, &SearchHandle);
+	if (ResultCode != EOS_EResult::EOS_Success)
 	{
-		const FFindSessions::Params& OpParams = Op.GetParams();
-
-		FOnlineError StateCheck = CheckFindSessionsState(OpParams);
-		if (StateCheck != Errors::Success())
-		{
-			Op.SetError(MoveTemp(StateCheck));
-			Promise.EmplaceValue();
-			return;
-		}
+		UE_LOG(LogTemp, Error, TEXT("[FSessionsEOSGS::FindSessions] EOS_Sessions_CreateSessionSearch failed with result [%s]"), *LexToString(ResultCode));
 
-		// Before we start the search, we reset the cache
-		SearchResultsUserMap.FindOrAdd(OpParams.LocalAccountId).Reset();
-		CurrentSessionSearchHandlesUserMap.Emplace(OpParams.LocalAccountId, Op.AsShared());
-
-		// We start preparing the search
-		EOS_Sessions_CreateSessionSearchOptions CreateSessionSearchOptions = { };
-		CreateSessionSearchOptions.ApiVersion = EOS_SESSIONS_CREATESESSIONSEARCH_API_LATEST;
-		static_assert(EOS_SESSIONS_CREATESESSIONSEARCH_API_LATEST == 1, "EOS_Sessions_CreateSessionSearchOptions updated, check new fields");
-		CreateSessionSearchOptions.MaxSearchResults = FMath::Clamp(OpParams.MaxResults, 0, EOS_SESSIONS_MAX_SEARCH_RESULTS);
-
-		EOS_HSessionSearch SearchHandle = nullptr;
-		EOS_EResult ResultCode = EOS_Sessions_CreateSessionSearch(SessionsHandle, &CreateSessionSearchOptions, &SearchHandle);
-		if (ResultCode != EOS_EResult::EOS_Success)
-		{
-			UE_LOG(LogTemp, Error, TEXT("[FSessionsEOSGS::FindSessions] EOS_Sessions_CreateSessionSearch failed with result [%s]"), *LexToString(ResultCode));
+		Promise.EmplaceValue(Errors::FromEOSResult(ResultCode));
+		return Future;
+	}
 
-			Op.SetError(Errors::FromEOSResult(ResultCode));
-			Promise.EmplaceValue();
-			return;
-		}
+	const TSharedRef<FSessionSearchHandleEOSGS>& CurrentSessionSearchHandleEOSGS = CurrentSessionSearchHandleEOSGSUserMap.Emplace(Params.LocalAccountId, MakeShared<FSessionSearchHandleEOSGS>(SearchHandle));
 
-		const TSharedRef<FSessionSearchHandleEOSGS>& CurrentSessionSearchHandleEOSGS = CurrentSessionSearchHandleEOSGSUserMap.Emplace(OpParams.LocalAccountId, MakeShared<FSessionSearchHandleEOSGS>(SearchHandle));
+	// Before we start the search, we reset the cache and save the promise
+	SearchResultsUserMap.FindOrAdd(Params.LocalAccountId).Reset();
+	CurrentSessionSearchPromisesUserMap.Emplace(Params.LocalAccountId, MoveTemp(Promise));
 
-		// We write the search attributes
-		WriteSessionSearchHandle(*CurrentSessionSearchHandleEOSGS, OpParams);
+	// We write the search attributes
+	WriteSessionSearchHandle(*CurrentSessionSearchHandleEOSGS, Params);
 
-		EOS_SessionSearch_FindOptions FindOptions = { };
-		FindOptions.ApiVersion = EOS_SESSIONSEARCH_FIND_API_LATEST;
-		static_assert(EOS_SESSIONSEARCH_FIND_API_LATEST == 2, "EOS_SessionSearch_FindOptions updated, check new fields");
-		FindOptions.LocalUserId = GetProductUserIdChecked(OpParams.LocalAccountId);
+	EOS_SessionSearch_FindOptions FindOptions = { };
+	FindOptions.ApiVersion = EOS_SESSIONSEARCH_FIND_API_LATEST;
+	static_assert(EOS_SESSIONSEARCH_FIND_API_LATEST == 2, "EOS_SessionSearch_FindOptions updated, check new fields");
+	FindOptions.LocalUserId = GetProductUserIdChecked(Params.LocalAccountId);
 
-		EOS_Async(EOS_SessionSearch_Find, CurrentSessionSearchHandleEOSGS->SearchHandle, FindOptions, MoveTemp(Promise));
-	})
-	.Then([this](TOnlineAsyncOp<FFindSessions>& Op, const EOS_SessionSearch_FindCallbackInfo*&& FindCallbackInfoResult)
-	{
-		if (FindCallbackInfoResult->ResultCode != EOS_EResult::EOS_Success)
+	EOS_Async(EOS_SessionSearch_Find, CurrentSessionSearchHandleEOSGS->SearchHandle, FindOptions,
+		[this, Params](const EOS_SessionSearch_FindCallbackInfo* FindCallbackInfoResult) mutable
 		{
-			Op.SetError(Errors::FromEOSResult(FindCallbackInfoResult->ResultCode));
-			return;
-		}
-
-		const FFindSessions::Params& OpParams = Op.GetParams();
+			TPromise<TOnlineResult<FFindSessions>>& Promise = CurrentSessionSearchPromisesUserMap.FindChecked(Params.LocalAccountId);
 
-		const TSharedRef<FSessionSearchHandleEOSGS>& CurrentSessionSearchHandleEOSGS = CurrentSessionSearchHandleEOSGSUserMap.FindChecked(OpParams.LocalAccountId);
-
-		// For a successful session, we'll get the search results
-		EOS_SessionSearch_GetSearchResultCountOptions GetSearchResultCountOptions = { };
-		GetSearchResultCountOptions.ApiVersion = EOS_SESSIONSEARCH_GETSEARCHRESULTCOUNT_API_LATEST;
-		static_assert(EOS_SESSIONSEARCH_GETSEARCHRESULTCOUNT_API_LATEST == 1, "EOS_SessionSearch_GetSearchResultCountOptions updated, check new fields");
-		int32 NumSearchResults = EOS_SessionSearch_GetSearchResultCount(CurrentSessionSearchHandleEOSGS->SearchHandle, &GetSearchResultCountOptions);
+			if (FindCallbackInfoResult->ResultCode != EOS_EResult::EOS_Success)
+			{
+				Promise.EmplaceValue(Errors::FromEOSResult(FindCallbackInfoResult->ResultCode));
+				return;
+			}
 
-		EOS_SessionSearch_CopySearchResultByIndexOptions CopySearchResultByIndexOptions = { };
-		CopySearchResultByIndexOptions.ApiVersion = EOS_SESSIONSEARCH_COPYSEARCHRESULTBYINDEX_API_LATEST;
-		static_assert(EOS_SESSIONSEARCH_COPYSEARCHRESULTBYINDEX_API_LATEST == 1, "EOS_SessionSearch_CopySearchResultByIndexOptions updated, check new fields");
+			const TSharedRef<FSessionSearchHandleEOSGS>& CurrentSessionSearchHandleEOSGS = CurrentSessionSearchHandleEOSGSUserMap.FindChecked(Params.LocalAccountId);
 
+			// For a successful session, we'll get the search results
+			EOS_SessionSearch_GetSearchResultCountOptions GetSearchResultCountOptions = { };
+			GetSearchResultCountOptions.ApiVersion = EOS_SESSIONSEARCH_GETSEARCHRESULTCOUNT_API_LATEST;
+			static_assert(EOS_SESSIONSEARCH_GETSEARCHRESULTCOUNT_API_LATEST == 1, "EOS_SessionSearch_GetSearchResultCountOptions updated, check new fields");
+			int32 NumSearchResults = EOS_SessionSearch_GetSearchResultCount(CurrentSessionSearchHandleEOSGS->SearchHandle, &GetSearchResultCountOptions);
 
-		TArray<TFuture<TDefaultErrorResult<FBuildSessionFromDetailsHandle>>> PendingSessionsBuilt;
+			EOS_SessionSearch_CopySearchResultByIndexOptions CopySearchResultByIndexOptions = { };
+			CopySearchResultByIndexOptions.ApiVersion = EOS_SESSIONSEARCH_COPYSEARCHRESULTBYINDEX_API_LATEST;
+			static_assert(EOS_SESSIONSEARCH_COPYSEARCHRESULTBYINDEX_API_LATEST == 1, "EOS_SessionSearch_CopySearchResultByIndexOptions updated, check new fields");
 
-		for (int32 Index = 0; Index < NumSearchResults; Index++)
-		{
-			EOS_HSessionDetails SessionDetailsHandle = nullptr;
+			TArray<TFuture<TDefaultErrorResult<FBuildSessionFromDetailsHandle>>> PendingSessionsBuilt;
 
-			CopySearchResultByIndexOptions.SessionIndex = Index;
-			EOS_EResult CopySearchResultByIndexResult = EOS_SessionSearch_CopySearchResultByIndex(CurrentSessionSearchHandleEOSGS->SearchHandle, &CopySearchResultByIndexOptions, &SessionDetailsHandle);
-			if (CopySearchResultByIndexResult == EOS_EResult::EOS_Success)
+			for (int32 Index = 0; Index < NumSearchResults; Index++)
 			{
-				TSharedRef<TPromise<TDefaultErrorResult<FBuildSessionFromDetailsHandle>>> BuildSessionPromise = MakeShared<TPromise<TDefaultErrorResult<FBuildSessionFromDetailsHandle>>>();
-				BuildSessionFromDetailsHandle({ OpParams.LocalAccountId, MakeShared<FSessionDetailsHandleEOSGS>(SessionDetailsHandle) })
-					.OnComplete([this, WeakOp = Op.AsWeak(), BuildSessionPromise](const TOnlineResult<FBuildSessionFromDetailsHandle>& Result) mutable
+				EOS_HSessionDetails SessionDetailsHandle = nullptr;
+
+				CopySearchResultByIndexOptions.SessionIndex = Index;
+				EOS_EResult CopySearchResultByIndexResult = EOS_SessionSearch_CopySearchResultByIndex(CurrentSessionSearchHandleEOSGS->SearchHandle, &CopySearchResultByIndexOptions, &SessionDetailsHandle);
+				if (CopySearchResultByIndexResult == EOS_EResult::EOS_Success)
 				{
-					if (Result.IsOk())
+					TSharedRef<TPromise<TDefaultErrorResult<FBuildSessionFromDetailsHandle>>> BuildSessionPromise = MakeShared<TPromise<TDefaultErrorResult<FBuildSessionFromDetailsHandle>>>();
+					BuildSessionFromDetailsHandle({ Params.LocalAccountId, MakeShared<FSessionDetailsHandleEOSGS>(SessionDetailsHandle) })
+						.OnComplete([this, BuildSessionPromise](const TOnlineResult<FBuildSessionFromDetailsHandle>& Result) mutable
 					{
-						if (TOnlineAsyncOpPtr<FFindSessions> StrongOp = WeakOp.Pin())
+						if (Result.IsOk())
 						{
 							AddSearchResult(Result.GetOkValue().Session, Result.GetOkValue().LocalAccountId);
 
 							BuildSessionPromise->EmplaceValue(Result.GetOkValue());
 						}
-					}
-					else
-					{
-						BuildSessionPromise->EmplaceValue(Result.GetErrorValue());
-					}
-				});
+						else
+						{
+							BuildSessionPromise->EmplaceValue(Result.GetErrorValue());
+						}
+					});
 			
-				PendingSessionsBuilt.Add(BuildSessionPromise->GetFuture());
-			}
-			else
-			{
-				UE_LOG(LogTemp, Error, TEXT("[FSessionsEOSGS::FindSessions] EOS_SessionSearch_CopySearchResultByIndex failed for index [%d] with result [%s]"), Index, *LexToString(CopySearchResultByIndexResult));
+					PendingSessionsBuilt.Add(BuildSessionPromise->GetFuture());
+				}
+				else
+				{
+					UE_LOG(LogTemp, Error, TEXT("[FSessionsEOSGS::FindSessions] EOS_SessionSearch_CopySearchResultByIndex failed for index [%d] with result [%s]"), Index, *LexToString(CopySearchResultByIndexResult));
 
-				Op.SetError(Errors::FromEOSResult(CopySearchResultByIndexResult));
+					Promise.EmplaceValue(Errors::FromEOSResult(CopySearchResultByIndexResult));
+					return;
+				}
 			}
-		}
 
-		WhenAll(MoveTemp(PendingSessionsBuilt))
-			.Next([this, WeakOp = Op.AsWeak()](TArray<TDefaultErrorResult<FBuildSessionFromDetailsHandle>>&& Results) mutable
-		{
-			if (TOnlineAsyncOpPtr<FFindSessions> StrongOp = WeakOp.Pin())
+			WhenAll(MoveTemp(PendingSessionsBuilt))
+				.Next([this, Params](TArray<TDefaultErrorResult<FBuildSessionFromDetailsHandle>>&& Results) mutable
 			{
+				TPromise<TOnlineResult<FFindSessions>>& Promise = CurrentSessionSearchPromisesUserMap.FindChecked(Params.LocalAccountId);
+
 				for (TDefaultErrorResult<FBuildSessionFromDetailsHandle>& Result : Results)
 				{
 					if (Result.IsError())
 					{
 						// Store first encountered error to return as result.
-						StrongOp->SetError(MoveTemp(Result.GetErrorValue()));
+						Promise.EmplaceValue(MoveTemp(Result.GetErrorValue()));
 						return;
 					}
 				}
 
-				const FFindSessions::Params& OpParams = StrongOp->GetParams();
-
-				TArray<FOnlineSessionId>& SearchResults = SearchResultsUserMap.FindChecked(OpParams.LocalAccountId);
-
-				StrongOp->SetResult({ SearchResults });
+				TArray<FOnlineSessionId>& SearchResults = SearchResultsUserMap.FindChecked(Params.LocalAccountId);
 
-				CurrentSessionSearchHandlesUserMap.Remove(OpParams.LocalAccountId);
-				CurrentSessionSearchHandleEOSGSUserMap.Remove(OpParams.LocalAccountId);
-			}
+				Promise.EmplaceValue(FFindSessions::Result{ SearchResults });
+				CurrentSessionSearchPromisesUserMap.Remove(Params.LocalAccountId);
+				CurrentSessionSearchHandleEOSGSUserMap.Remove(Params.LocalAccountId);
+			});
 		});
-	});
-
-	Op->Enqueue(GetSerialQueue());
 
-	return Op->GetHandle();
+	return Future;
 }
 
-TOnlineAsyncOpHandle<FJoinSession> FSessionsEOSGS::JoinSession(FJoinSession::Params&& Params)
+TFuture<TOnlineResult<FJoinSession>> FSessionsEOSGS::JoinSessionImpl(const FJoinSession::Params& Params)
 {
-	TOnlineAsyncOpRef<FJoinSession> Op = GetOp<FJoinSession>(MoveTemp(Params));
-	const FJoinSession::Params& OpParams = Op->GetParams();
+	TPromise<TOnlineResult<FJoinSession>> Promise;
+	TFuture<TOnlineResult<FJoinSession>> Future = Promise.GetFuture();
 
-	TOnlineResult<FGetSessionById> GetSessionByIdResult = GetSessionById({ OpParams.SessionId });
+	TOnlineResult<FGetSessionById> GetSessionByIdResult = GetSessionById({ Params.SessionId });
 	if (GetSessionByIdResult.IsError())
 	{
 		// If no result is found, the id might be expired, which we should notify
-		if (FOnlineSessionIdRegistryEOSGS::Get().IsSessionIdExpired(OpParams.SessionId))
+		if (FOnlineSessionIdRegistryEOSGS::Get().IsSessionIdExpired(Params.SessionId))
 		{
-			UE_LOG(LogTemp, Warning, TEXT("[FSessionsEOSGS::JoinSession] SessionId parameter [%s] is expired. Please call FindSessions to get an updated list of available sessions "), *ToLogString(OpParams.SessionId));
+			UE_LOG(LogTemp, Warning, TEXT("[%s] SessionId parameter [%s] is expired. Please call FindSessions to get an updated list of available sessions "), UTF8_TO_TCHAR(__FUNCTION__), *ToLogString(Params.SessionId));
 		}
 
-		Op->SetError(MoveTemp(GetSessionByIdResult.GetErrorValue()));
-		return Op->GetHandle();
+		Promise.EmplaceValue(GetSessionByIdResult.GetErrorValue());
+		return Future;
 	}
 
 	const TSharedRef<const ISession>& FoundSession = GetSessionByIdResult.GetOkValue().Session;
@@ -1248,7 +1182,7 @@ TOnlineAsyncOpHandle<FJoinSession> FSessionsEOSGS::JoinSession(FJoinSession::Par
 	// LAN Sessions
 	if (FoundSession->GetSessionInfo().bIsLANSession)
 	{
-		return FSessionsLAN::JoinSession(MoveTemp(Params));
+		return FSessionsLAN::JoinSessionImpl(Params);
 	}
 
 	// EOSGS Sessions
@@ -1257,106 +1191,65 @@ TOnlineAsyncOpHandle<FJoinSession> FSessionsEOSGS::JoinSession(FJoinSession::Par
 	const FSessionEOSGS& SessionEOSGS = FSessionEOSGS::Cast(*FoundSession);
 	if (!SessionEOSGS.SessionDetailsHandle.IsValid())
 	{
-		UE_LOG(LogTemp, Warning, TEXT("[FSessionsEOSGS::FJoinSession] Could not join session with invalid session details handle in session with id [%s]"), *ToLogString(OpParams.SessionId));
-
-		Op->SetError(Errors::InvalidParams());
-		return Op->GetHandle();
-	}
+		UE_LOG(LogTemp, Warning, TEXT("[%s] Could not join session with invalid session details handle in session with id [%s]"), UTF8_TO_TCHAR(__FUNCTION__), *ToLogString(Params.SessionId));
 
-	FOnlineError ParamsCheck = CheckJoinSessionParams(OpParams);
-	if (ParamsCheck != Errors::Success())
-	{
-		Op->SetError(MoveTemp(ParamsCheck));
-		return Op->GetHandle();
+		Promise.EmplaceValue(Errors::InvalidState());
+		return Future;
 	}
 
-	// We start the join operation
-	Op->Then([this](TOnlineAsyncOp<FJoinSession>& Op, TPromise<const EOS_Sessions_JoinSessionCallbackInfo*>&& Promise)
-	{
-		const FJoinSession::Params& OpParams = Op.GetParams();
-
-		FOnlineError StateCheck = CheckJoinSessionState(OpParams);
-		if (StateCheck != Errors::Success())
-		{
-			Op.SetError(MoveTemp(StateCheck));
-			Promise.EmplaceValue();
-			return;
-		}
-
-		TOnlineResult<FGetSessionById> GetSessionByIdResult = GetSessionById({ OpParams.SessionId });
-		if (GetSessionByIdResult.IsError())
-		{
-			// If no result is found, the id might be expired, which we should notify
-			if (FOnlineSessionIdRegistryEOSGS::Get().IsSessionIdExpired(OpParams.SessionId))
-			{
-				UE_LOG(LogTemp, Warning, TEXT("[FSessionsEOSGS::JoinSession] SessionId parameter [%s] is expired. Please call FindSessions to get an updated list of available sessions "), *ToLogString(OpParams.SessionId));
-			}
-
-			Op.SetError(MoveTemp(GetSessionByIdResult.GetErrorValue()));
-			Promise.EmplaceValue();
-			return;
-		}
-
-		const TSharedRef<const ISession>& FoundSession = GetSessionByIdResult.GetOkValue().Session;
-
-		// We start setup for the API call
-		EOS_Sessions_JoinSessionOptions JoinSessionOptions = { };
-		JoinSessionOptions.ApiVersion = EOS_SESSIONS_JOINSESSION_API_LATEST;
-		static_assert(EOS_SESSIONS_JOINSESSION_API_LATEST == 2, "EOS_Sessions_JoinSessionOptions updated, check new fields");
-
-		JoinSessionOptions.bPresenceEnabled = OpParams.bPresenceEnabled;
+	// We start setup for the API call
+	EOS_Sessions_JoinSessionOptions JoinSessionOptions = { };
+	JoinSessionOptions.ApiVersion = EOS_SESSIONS_JOINSESSION_API_LATEST;
+	static_assert(EOS_SESSIONS_JOINSESSION_API_LATEST == 2, "EOS_Sessions_JoinSessionOptions updated, check new fields");
 
-		JoinSessionOptions.LocalUserId = GetProductUserIdChecked(OpParams.LocalAccountId);
+	JoinSessionOptions.bPresenceEnabled = Params.bPresenceEnabled;
 
-		const FTCHARToUTF8 SessionNameUtf8(*OpParams.SessionName.ToString());
-		JoinSessionOptions.SessionName = SessionNameUtf8.Get();
+	JoinSessionOptions.LocalUserId = GetProductUserIdChecked(Params.LocalAccountId);
 
-		const FSessionEOSGS& SessionEOSGS = FSessionEOSGS::Cast(*FoundSession);
-		JoinSessionOptions.SessionHandle = SessionEOSGS.SessionDetailsHandle->SessionDetailsHandle;
+	const FTCHARToUTF8 SessionNameUtf8(*Params.SessionName.ToString());
+	JoinSessionOptions.SessionName = SessionNameUtf8.Get();
 
-		EOS_Async(EOS_Sessions_JoinSession, SessionsHandle, JoinSessionOptions, MoveTemp(Promise));
-	})
-	.Then([this](TOnlineAsyncOp<FJoinSession>& Op, const EOS_Sessions_JoinSessionCallbackInfo*&& Result)
-	{
-		const FJoinSession::Params& OpParams = Op.GetParams();
+	JoinSessionOptions.SessionHandle = SessionEOSGS.SessionDetailsHandle->SessionDetailsHandle;
 
-		if (Result->ResultCode != EOS_EResult::EOS_Success)
+	EOS_Async(EOS_Sessions_JoinSession, SessionsHandle, JoinSessionOptions, 
+		[this, Promise = MoveTemp(Promise), Params](const EOS_Sessions_JoinSessionCallbackInfo* Result) mutable
 		{
-			Op.SetError(Errors::FromEOSResult(Result->ResultCode));
-			return;
-		}
-		
-		TOnlineResult<FGetSessionById> GetSessionByIdResult = GetSessionById({ OpParams.SessionId });
-		if (GetSessionByIdResult.IsError())
-		{
-			// If no result is found, the id might be expired, which we should notify
-			if (FOnlineSessionIdRegistryEOSGS::Get().IsSessionIdExpired(OpParams.SessionId))
+			if (Result->ResultCode != EOS_EResult::EOS_Success)
 			{
-				UE_LOG(LogTemp, Warning, TEXT("[FSessionsEOSGS::JoinSession] SessionId parameter [%s] is expired. Please call FindSessions to get an updated list of available sessions "), *ToLogString(OpParams.SessionId));
+				Promise.EmplaceValue(Errors::FromEOSResult(Result->ResultCode));
+				return;
 			}
+		
+			TOnlineResult<FGetSessionById> GetSessionByIdResult = GetSessionById({ Params.SessionId });
+			if (GetSessionByIdResult.IsError())
+			{
+				// If no result is found, the id might be expired, which we should notify
+				if (FOnlineSessionIdRegistryEOSGS::Get().IsSessionIdExpired(Params.SessionId))
+				{
+					UE_LOG(LogTemp, Warning, TEXT("[%s] SessionId parameter [%s] is expired. Please call FindSessions to get an updated list of available sessions "), UTF8_TO_TCHAR(__FUNCTION__), *ToLogString(Params.SessionId));
+				}
 
-			Op.SetError(MoveTemp(GetSessionByIdResult.GetErrorValue()));
-			return;
-		}
+				Promise.EmplaceValue(MoveTemp(GetSessionByIdResult.GetErrorValue()));
+				return;
+			}
 
-		const TSharedRef<const ISession>& FoundSession = GetSessionByIdResult.GetOkValue().Session;
+			const TSharedRef<const ISession>& FoundSession = GetSessionByIdResult.GetOkValue().Session;
 
-		AddSessionReferences(FoundSession->GetSessionId(), OpParams.SessionName, OpParams.LocalAccountId, OpParams.bPresenceEnabled);
+			AddSessionReferences(FoundSession->GetSessionId(), Params.SessionName, Params.LocalAccountId, Params.bPresenceEnabled);
 
-		// After successfully joining a session, we'll remove all related invites if any are found
-		ClearSessionInvitesForSession(OpParams.LocalAccountId, FoundSession->GetSessionId());
+			// After successfully joining a session, we'll remove all related invites if any are found
+			ClearSessionInvitesForSession(Params.LocalAccountId, FoundSession->GetSessionId());
 
-		Op.SetResult(FJoinSession::Result{ });
+			Promise.EmplaceValue(FJoinSession::Result{ });
 
-		FSessionJoined Event = { { OpParams.LocalAccountId }, FoundSession->GetSessionId() };
+			FSessionJoined Event = { { Params.LocalAccountId }, FoundSession->GetSessionId() };
 
-		SessionEvents.OnSessionJoined.Broadcast(Event);
+			SessionEvents.OnSessionJoined.Broadcast(Event);
 
-		// A successful join allows the client to server travel, after which AddSessionMember will be called by the engine
-	})
-	.Enqueue(GetSerialQueue());
+			// A successful join allows the client to server travel, after which AddSessionMember will be called by the engine
+		});
 
-	return Op->GetHandle();
+	return Future;
 }
 
 TOnlineAsyncOpHandle<FBuildSessionFromDetailsHandle> FSessionsEOSGS::BuildSessionFromInvite(const FAccountId& LocalAccountId, const FString& InInviteId)
@@ -1499,244 +1392,171 @@ TOnlineAsyncOpHandle<FBuildSessionFromDetailsHandle> FSessionsEOSGS::BuildSessio
 	return Op->GetHandle();
 }
 
-TOnlineAsyncOpHandle<FSendSessionInvite> FSessionsEOSGS::SendSessionInvite(FSendSessionInvite::Params&& Params)
+TFuture<TOnlineResult<FSendSessionInvite>> FSessionsEOSGS::SendSessionInviteImpl(const FSendSessionInvite::Params& Params)
 {
-	TOnlineAsyncOpRef<FSendSessionInvite> Op = GetOp<FSendSessionInvite>(MoveTemp(Params));
+	TPromise<TOnlineResult<FSendSessionInvite>> Promise;
+	TFuture<TOnlineResult<FSendSessionInvite>> Future = Promise.GetFuture();
 
-	Op->Then([this](TOnlineAsyncOp<FSendSessionInvite>& Op, TPromise<const EOS_Sessions_SendInviteCallbackInfo*>&& Promise)
+	TArray<TFuture<TDefaultErrorResult<FSendSingleSessionInviteImpl>>> PendingSessionInvites;
+	for (const FAccountId& TargetAccountId : Params.TargetUsers)
 	{
-		const FSendSessionInvite::Params& OpParams = Op.GetParams();
+		FSendSingleSessionInviteImpl::Params FSendSingleSessionInviteParams;
+		FSendSingleSessionInviteParams.LocalAccountId = Params.LocalAccountId;
+		FSendSingleSessionInviteParams.SessionName = Params.SessionName;
+		FSendSingleSessionInviteParams.TargetAccountId = TargetAccountId;
 
-		FOnlineError StateCheck = CheckSendSessionInviteState(OpParams);
-		if (StateCheck != Errors::Success())
-		{
-			Op.SetError(MoveTemp(StateCheck));
-			Promise.EmplaceValue();
-			return;
-		}
-
-		TArray<TFuture<TDefaultErrorResult<FSendSingleSessionInviteImpl>>> PendingSessionInvites;
-		for (const FAccountId& TargetAccountId : OpParams.TargetUsers)
-		{
-			FSendSingleSessionInviteImpl::Params FSendSingleSessionInviteParams;
-			FSendSingleSessionInviteParams.LocalAccountId = OpParams.LocalAccountId;
-			FSendSingleSessionInviteParams.SessionName = OpParams.SessionName;
-			FSendSingleSessionInviteParams.TargetAccountId = TargetAccountId;
-
-			TSharedRef<TPromise<TDefaultErrorResult<FSendSingleSessionInviteImpl>>> SessionInvitePromise = MakeShared<TPromise<TDefaultErrorResult<FSendSingleSessionInviteImpl>>>();
-			SendSingleSessionInviteImpl(MoveTemp(FSendSingleSessionInviteParams))
-				.OnComplete([SessionInvitePromise](const TOnlineResult<FSendSingleSessionInviteImpl>& Result)
+		TSharedRef<TPromise<TDefaultErrorResult<FSendSingleSessionInviteImpl>>> SessionInvitePromise = MakeShared<TPromise<TDefaultErrorResult<FSendSingleSessionInviteImpl>>>();
+		SendSingleSessionInviteImpl(MoveTemp(FSendSingleSessionInviteParams))
+			.OnComplete([SessionInvitePromise](const TOnlineResult<FSendSingleSessionInviteImpl>& Result)
+				{
+					if (Result.IsOk())
 					{
-						if (Result.IsOk())
-						{
-							SessionInvitePromise->EmplaceValue(Result.GetOkValue());
-						}
-						else
-						{
-							SessionInvitePromise->EmplaceValue(Result.GetErrorValue());
-						}
-					});
+						SessionInvitePromise->EmplaceValue(Result.GetOkValue());
+					}
+					else
+					{
+						SessionInvitePromise->EmplaceValue(Result.GetErrorValue());
+					}
+				});
 
-			PendingSessionInvites.Emplace(SessionInvitePromise->GetFuture());
-		}
+		PendingSessionInvites.Emplace(SessionInvitePromise->GetFuture());
+	}
 
-		WhenAll(MoveTemp(PendingSessionInvites))
-			.Next([WeakOp = Op.AsWeak()](TArray<TDefaultErrorResult<FSendSingleSessionInviteImpl>>&& Results) mutable
+	WhenAll(MoveTemp(PendingSessionInvites))
+		.Next([Promise = MoveTemp(Promise)](TArray<TDefaultErrorResult<FSendSingleSessionInviteImpl>>&& Results) mutable
+	{
+		for (TDefaultErrorResult<FSendSingleSessionInviteImpl>& Result : Results)
 		{
-			if (TOnlineAsyncOpPtr<FSendSessionInvite> StrongOp = WeakOp.Pin())
+			if (Result.IsError())
 			{
-				for (TDefaultErrorResult<FSendSingleSessionInviteImpl>& Result : Results)
-				{
-					if (Result.IsError())
-					{
-						StrongOp->SetError(MoveTemp(Result.GetErrorValue()));
-						return;
-					}
-				}
-
-				StrongOp->SetResult(FSendSessionInvite::Result{ });
+				Promise.EmplaceValue(MoveTemp(Result.GetErrorValue()));
+				return;
 			}
-		});
+		}
+
+		Promise.EmplaceValue(FSendSessionInvite::Result{ });
 	});
 
-	return Op->GetHandle();
+	return Future;
 }
 
-TOnlineAsyncOpHandle<FRejectSessionInvite> FSessionsEOSGS::RejectSessionInvite(FRejectSessionInvite::Params&& Params)
+TFuture<TOnlineResult<FRejectSessionInvite>> FSessionsEOSGS::RejectSessionInviteImpl(const FRejectSessionInvite::Params& Params)
 {
-	TOnlineAsyncOpRef<FRejectSessionInvite> Op = GetOp<FRejectSessionInvite>(MoveTemp(Params));
-
-	Op->Then([this](TOnlineAsyncOp<FRejectSessionInvite>& Op, TPromise<const EOS_Sessions_RejectInviteCallbackInfo*>&& Promise)
-	{
-		const FRejectSessionInvite::Params& OpParams = Op.GetParams();
-
-		FOnlineError StateCheck = CheckRejectSessionInviteState(OpParams);
-		if (StateCheck != Errors::Success())
-		{
-			Op.SetError(MoveTemp(StateCheck));
-			Promise.EmplaceValue();
-			return;
-		}
+	TPromise<TOnlineResult<FRejectSessionInvite>> Promise;
+	TFuture<TOnlineResult<FRejectSessionInvite>> Future = Promise.GetFuture();
 
-		EOS_Sessions_RejectInviteOptions RejectInviteOptions = { };
-		RejectInviteOptions.ApiVersion = EOS_SESSIONS_REJECTINVITE_API_LATEST;
-		static_assert(EOS_SESSIONS_REJECTINVITE_API_LATEST == 1, "EOS_Sessions_RejectInviteOptions updated, check new fields");
+	EOS_Sessions_RejectInviteOptions RejectInviteOptions = { };
+	RejectInviteOptions.ApiVersion = EOS_SESSIONS_REJECTINVITE_API_LATEST;
+	static_assert(EOS_SESSIONS_REJECTINVITE_API_LATEST == 1, "EOS_Sessions_RejectInviteOptions updated, check new fields");
 
-		RejectInviteOptions.LocalUserId = GetProductUserIdChecked(OpParams.LocalAccountId);
+	RejectInviteOptions.LocalUserId = GetProductUserIdChecked(Params.LocalAccountId);
 
-		const FString& InviteIdStr = FOnlineSessionInviteIdRegistryEOSGS::Get().BasicRegistry.FindIdValue(OpParams.SessionInviteId);
-		const FTCHARToUTF8 InviteIdUtf8(*InviteIdStr);
-		RejectInviteOptions.InviteId = InviteIdUtf8.Get();
+	const FString& InviteIdStr = FOnlineSessionInviteIdRegistryEOSGS::Get().BasicRegistry.FindIdValue(Params.SessionInviteId);
+	const FTCHARToUTF8 InviteIdUtf8(*InviteIdStr);
+	RejectInviteOptions.InviteId = InviteIdUtf8.Get();
 
-		EOS_Async(EOS_Sessions_RejectInvite, SessionsHandle, RejectInviteOptions, MoveTemp(Promise));
-		return;
-	})
-	.Then([this](TOnlineAsyncOp<FRejectSessionInvite>& Op, const EOS_Sessions_RejectInviteCallbackInfo* Result)
-	{
-		if (Result->ResultCode != EOS_EResult::EOS_Success)
+	EOS_Async(EOS_Sessions_RejectInvite, SessionsHandle, RejectInviteOptions, 
+		[this, Promise = MoveTemp(Promise), Params](const EOS_Sessions_RejectInviteCallbackInfo* Result) mutable
 		{
-			Op.SetError(Errors::FromEOSResult(Result->ResultCode));
-			return;
-		}
-
-		const FRejectSessionInvite::Params& OpParams = Op.GetParams();
+			if (Result->ResultCode != EOS_EResult::EOS_Success)
+			{
+				Promise.EmplaceValue(Errors::FromEOSResult(Result->ResultCode));
+				return;
+			}
 
-		if (TMap<FSessionInviteId, TSharedRef<FSessionInvite>>* UserMap = SessionInvitesUserMap.Find(OpParams.LocalAccountId))
-		{
-			UserMap->Remove(OpParams.SessionInviteId);
-		}
+			if (TMap<FSessionInviteId, TSharedRef<FSessionInvite>>* UserMap = SessionInvitesUserMap.Find(Params.LocalAccountId))
+			{
+				UserMap->Remove(Params.SessionInviteId);
+			}
 
-		Op.SetResult(FRejectSessionInvite::Result{ });
-	})
-	.Enqueue(GetSerialQueue()); // TODO: Use the parallel queue instead when possible
+			Promise.EmplaceValue(FRejectSessionInvite::Result{ });
+		});
 
-	return Op->GetHandle();
+	return Future;
 }
 
-TOnlineAsyncOpHandle<FAddSessionMember> FSessionsEOSGS::AddSessionMember(FAddSessionMember::Params&& Params)
+TFuture<TOnlineResult<FAddSessionMember>> FSessionsEOSGS::AddSessionMemberImpl(const FAddSessionMember::Params& Params)
 {
 	// LAN Sessions
 	TOnlineResult<FGetSessionByName> Result = GetSessionByName({ Params.SessionName });
 	if (Result.IsOk() && Result.GetOkValue().Session->GetSessionInfo().bIsLANSession)
 	{
-		return FSessionsLAN::AddSessionMember(MoveTemp(Params));
+		return FSessionsLAN::AddSessionMemberImpl(Params);
 	}
 
 	// EOSGS Sessions
 
-	TOnlineAsyncOpRef<FAddSessionMember> Op = GetOp<FAddSessionMember>(MoveTemp(Params));
+	TPromise<TOnlineResult<FAddSessionMember>> Promise;
+	TFuture<TOnlineResult<FAddSessionMember>> Future = Promise.GetFuture();
 
-	Op->Then([this](TOnlineAsyncOp<FAddSessionMember>& Op, TPromise<const EOS_Sessions_RegisterPlayersCallbackInfo*>&& Promise)
-	{
-		const FAddSessionMember::Params& OpParams = Op.GetParams();
+	EOS_Sessions_RegisterPlayersOptions Options = { };
+	Options.ApiVersion = EOS_SESSIONS_REGISTERPLAYERS_API_LATEST;
+	static_assert(EOS_SESSIONS_REGISTERPLAYERS_API_LATEST == 2, "EOS_Sessions_RegisterPlayersOptions updated, check new fields");
 
-		FOnlineError StateCheck = CheckAddSessionMemberState(OpParams);
-		if (StateCheck != Errors::Success())
-		{
-			Op.SetError(MoveTemp(StateCheck));
-			Promise.EmplaceValue();
-			return;
-		}
-
-		EOS_Sessions_RegisterPlayersOptions Options = { };
-		Options.ApiVersion = EOS_SESSIONS_REGISTERPLAYERS_API_LATEST;
-		static_assert(EOS_SESSIONS_REGISTERPLAYERS_API_LATEST == 2, "EOS_Sessions_RegisterPlayersOptions updated, check new fields");
-
-		Options.PlayersToRegisterCount = 1;
+	Options.PlayersToRegisterCount = 1;
 
-		// TODO: Do this with ResolveAccountIds instead
-		TArray<EOS_ProductUserId> PlayersToRegister = { GetProductUserIdChecked(OpParams.LocalAccountId) };
-		Options.PlayersToRegister = PlayersToRegister.GetData();
+	// TODO: Do this with ResolveAccountIds instead
+	TArray<EOS_ProductUserId> PlayersToRegister = { GetProductUserIdChecked(Params.LocalAccountId) };
+	Options.PlayersToRegister = PlayersToRegister.GetData();
 
-		const FTCHARToUTF8 SessionNameUtf8(*OpParams.SessionName.ToString());
-		Options.SessionName = SessionNameUtf8.Get();
+	const FTCHARToUTF8 SessionNameUtf8(*Params.SessionName.ToString());
+	Options.SessionName = SessionNameUtf8.Get();
 
-		EOS_Async(EOS_Sessions_RegisterPlayers, SessionsHandle, Options, MoveTemp(Promise));
-	})
-	.Then([this](TOnlineAsyncOp<FAddSessionMember>& Op, const EOS_Sessions_RegisterPlayersCallbackInfo* Result)
-	{
-		if (Result->ResultCode != EOS_EResult::EOS_Success)
+	EOS_Async(EOS_Sessions_RegisterPlayers, SessionsHandle, Options,
+		[this, Promise = MoveTemp(Promise), Params](const EOS_Sessions_RegisterPlayersCallbackInfo* Result) mutable
 		{
-			Op.SetError(Errors::FromEOSResult(Result->ResultCode));
-			return;
-		}
+			if (Result->ResultCode != EOS_EResult::EOS_Success)
+			{
+				Promise.EmplaceValue(Errors::FromEOSResult(Result->ResultCode));
+				return;
+			}
 
-		TOnlineResult<FAddSessionMember> ImplResult = AddSessionMemberImpl(Op.GetParams());
-		if (ImplResult.IsOk())
-		{
-			Op.SetResult(MoveTemp(ImplResult.GetOkValue()));
-		}
-		else
-		{
-			Op.SetError(MoveTemp(ImplResult.GetErrorValue()));
-		}
-	})
-	.Enqueue(GetSerialQueue());
+			Promise.EmplaceValue(AddSessionMemberInternal(Params));
+		});
 
-	return Op->GetHandle();
+	return Future;
 }
 
-TOnlineAsyncOpHandle<FRemoveSessionMember> FSessionsEOSGS::RemoveSessionMember(FRemoveSessionMember::Params&& Params)
+TFuture<TOnlineResult<FRemoveSessionMember>> FSessionsEOSGS::RemoveSessionMemberImpl(const FRemoveSessionMember::Params& Params)
 {
 	// LAN Sessions
 	TOnlineResult<FGetSessionByName> Result = GetSessionByName({ Params.SessionName });
 	if (Result.IsOk() && Result.GetOkValue().Session->GetSessionInfo().bIsLANSession)
 	{
-		return FSessionsLAN::RemoveSessionMember(MoveTemp(Params));
+		return FSessionsLAN::RemoveSessionMemberImpl(Params);
 	}
 
 	// EOSGS Sessions
 
-	TOnlineAsyncOpRef<FRemoveSessionMember> Op = GetOp<FRemoveSessionMember>(MoveTemp(Params));
-
-	Op->Then([this](TOnlineAsyncOp<FRemoveSessionMember>& Op, TPromise<const EOS_Sessions_UnregisterPlayersCallbackInfo*>&& Promise)
-	{
-		const FRemoveSessionMember::Params& OpParams = Op.GetParams();
-
-		FOnlineError StateCheck = CheckRemoveSessionMemberState(OpParams);
-		if (StateCheck != Errors::Success())
-		{
-			Op.SetError(MoveTemp(StateCheck));
-			Promise.EmplaceValue();
-			return;
-		}
+	TPromise<TOnlineResult<FRemoveSessionMember>> Promise;
+	TFuture<TOnlineResult<FRemoveSessionMember>> Future = Promise.GetFuture();
 
-		EOS_Sessions_UnregisterPlayersOptions Options = { };
-		Options.ApiVersion = EOS_SESSIONS_UNREGISTERPLAYERS_API_LATEST;
-		static_assert(EOS_SESSIONS_UNREGISTERPLAYERS_API_LATEST == 2, "EOS_Sessions_UnregisterPlayersOptions updated, check new fields");
+	EOS_Sessions_UnregisterPlayersOptions Options = { };
+	Options.ApiVersion = EOS_SESSIONS_UNREGISTERPLAYERS_API_LATEST;
+	static_assert(EOS_SESSIONS_UNREGISTERPLAYERS_API_LATEST == 2, "EOS_Sessions_UnregisterPlayersOptions updated, check new fields");
 
-		Options.PlayersToUnregisterCount = 1;
+	Options.PlayersToUnregisterCount = 1;
 
-		// TODO: Do this with ResolveAccountIds instead
-		TArray<EOS_ProductUserId> PlayersToUnregister{ GetProductUserIdChecked(OpParams.LocalAccountId) };
-		Options.PlayersToUnregister = PlayersToUnregister.GetData();
+	// TODO: Do this with ResolveAccountIds instead
+	TArray<EOS_ProductUserId> PlayersToUnregister{ GetProductUserIdChecked(Params.LocalAccountId) };
+	Options.PlayersToUnregister = PlayersToUnregister.GetData();
 
-		const FTCHARToUTF8 SessionNameUtf8(*OpParams.SessionName.ToString());
-		Options.SessionName = SessionNameUtf8.Get();
+	const FTCHARToUTF8 SessionNameUtf8(*Params.SessionName.ToString());
+	Options.SessionName = SessionNameUtf8.Get();
 
-		EOS_Async(EOS_Sessions_UnregisterPlayers, SessionsHandle, Options, MoveTemp(Promise));
-	})
-	.Then([this](TOnlineAsyncOp<FRemoveSessionMember>& Op, const EOS_Sessions_UnregisterPlayersCallbackInfo* Result)
-	{
-		if (Result->ResultCode != EOS_EResult::EOS_Success)
+	EOS_Async(EOS_Sessions_UnregisterPlayers, SessionsHandle, Options, 
+		[this, Promise = MoveTemp(Promise), Params](const EOS_Sessions_UnregisterPlayersCallbackInfo* Result) mutable
 		{
-			Op.SetError(Errors::FromEOSResult(Result->ResultCode));
-			return;
-		}
+			if (Result->ResultCode != EOS_EResult::EOS_Success)
+			{
+				Promise.EmplaceValue(Errors::FromEOSResult(Result->ResultCode));
+				return;
+			}
 
-		TOnlineResult<FRemoveSessionMember> ImplResult = RemoveSessionMemberImpl(Op.GetParams());
-		if (ImplResult.IsOk())
-		{
-			Op.SetResult(MoveTemp(ImplResult.GetOkValue()));
-		}
-		else
-		{
-			Op.SetError(MoveTemp(ImplResult.GetErrorValue()));
-		}
-	})
-	.Enqueue(GetSerialQueue());
+			Promise.EmplaceValue(RemoveSessionMemberInternal(Params));
+		});
 
-	return Op->GetHandle();
+	return Future;
 }
 
 /** FSessionsLAN */
diff --git a/Engine/Plugins/Online/Experimental/OnlineServicesEOSGS/Source/Public/Online/SessionsEOSGS.h b/Engine/Plugins/Online/Experimental/OnlineServicesEOSGS/Source/Public/Online/SessionsEOSGS.h
--- a/Engine/Plugins/Online/Experimental/OnlineServicesEOSGS/Source/Public/Online/SessionsEOSGS.h
+++ b/Engine/Plugins/Online/Experimental/OnlineServicesEOSGS/Source/Public/Online/SessionsEOSGS.h
@@ -189,18 +189,16 @@ public:
 	void Initialize() override;
 	void Shutdown() override;
 
-	// ISessions
-	virtual TOnlineAsyncOpHandle<FCreateSession> CreateSession(FCreateSession::Params&& Params) override;
-	virtual TOnlineAsyncOpHandle<FLeaveSession> LeaveSession(FLeaveSession::Params&& Params) override;
-	virtual TOnlineAsyncOpHandle<FFindSessions> FindSessions(FFindSessions::Params&& Params) override;
-	virtual TOnlineAsyncOpHandle<FJoinSession> JoinSession(FJoinSession::Params&& Params) override;
-	virtual TOnlineAsyncOpHandle<FSendSessionInvite> SendSessionInvite(FSendSessionInvite::Params&& Params) override;
-	virtual TOnlineAsyncOpHandle<FRejectSessionInvite> RejectSessionInvite(FRejectSessionInvite::Params&& Params) override;
-	virtual TOnlineAsyncOpHandle<FAddSessionMember> AddSessionMember(FAddSessionMember::Params&& Params) override;
-	virtual TOnlineAsyncOpHandle<FRemoveSessionMember> RemoveSessionMember(FRemoveSessionMember::Params&& Params) override;
-
 	// FSessionsCommon
-	virtual TFuture<TOnlineResult<FUpdateSessionSettingsImpl>> UpdateSessionSettingsImpl(FUpdateSessionSettingsImpl::Params&& Params) override;
+	virtual TFuture<TOnlineResult<FCreateSession>> CreateSessionImpl(const FCreateSession::Params& Params) override;
+	virtual TFuture<TOnlineResult<FUpdateSessionSettings>> UpdateSessionSettingsImpl(const FUpdateSessionSettings::Params& Params) override;
+	virtual TFuture<TOnlineResult<FLeaveSession>> LeaveSessionImpl(const FLeaveSession::Params& Params) override;
+	virtual TFuture<TOnlineResult<FFindSessions>> FindSessionsImpl(const FFindSessions::Params& Params) override;
+	virtual TFuture<TOnlineResult<FJoinSession>> JoinSessionImpl(const FJoinSession::Params& Params) override;
+	virtual TFuture<TOnlineResult<FAddSessionMember>> AddSessionMemberImpl(const FAddSessionMember::Params& Params) override;
+	virtual TFuture<TOnlineResult<FRemoveSessionMember>> RemoveSessionMemberImpl(const FRemoveSessionMember::Params& Params) override;
+	virtual TFuture<TOnlineResult<FSendSessionInvite>> SendSessionInviteImpl(const FSendSessionInvite::Params& Params) override;
+	virtual TFuture<TOnlineResult<FRejectSessionInvite>> RejectSessionInviteImpl(const FRejectSessionInvite::Params& Params) override;
 
 protected:
 	void RegisterEventHandlers();
diff --git a/Engine/Plugins/Online/Experimental/OnlineServicesOSSAdapter/Source/Private/Online/SessionsOSSAdapter.cpp b/Engine/Plugins/Online/Experimental/OnlineServicesOSSAdapter/Source/Private/Online/SessionsOSSAdapter.cpp
--- a/Engine/Plugins/Online/Experimental/OnlineServicesOSSAdapter/Source/Private/Online/SessionsOSSAdapter.cpp
+++ b/Engine/Plugins/Online/Experimental/OnlineServicesOSSAdapter/Source/Private/Online/SessionsOSSAdapter.cpp
@@ -393,68 +393,42 @@ void FSessionsOSSAdapter::Shutdown()
 	Super::Shutdown();
 }
 
-TOnlineAsyncOpHandle<FCreateSession> FSessionsOSSAdapter::CreateSession(FCreateSession::Params&& Params)
+TFuture<TOnlineResult<FCreateSession>> FSessionsOSSAdapter::CreateSessionImpl(const FCreateSession::Params& Params)
 {
-	TOnlineAsyncOpRef<FCreateSession> Op = GetOp<FCreateSession>(MoveTemp(Params));
-	const FCreateSession::Params& OpParams = Op->GetParams();
 
-	FOnlineError ParamsCheck = CheckCreateSessionParams(OpParams);
-	if (ParamsCheck != Errors::Success())
-	{
-		Op->SetError(MoveTemp(ParamsCheck));
-		return Op->GetHandle();
-	}
+	TPromise<TOnlineResult<FCreateSession>> Promise;
+	TFuture<TOnlineResult<FCreateSession>> Future = Promise.GetFuture();
 
-	Op->Then([this](TOnlineAsyncOp<FCreateSession>& Op)
-	{
-		const FCreateSession::Params& OpParams = Op.GetParams();
+	// Dedicated servers can update sessions so we won't check if LocalAccountId is valid
+	const FUniqueNetIdRef LocalAccountId = Services.Get<FAuthOSSAdapter>()->GetUniqueNetId(Params.LocalAccountId).ToSharedRef();
 
-		FOnlineError StateCheck = CheckCreateSessionState(OpParams);
-		if (StateCheck != Errors::Success())
+	MakeMulticastAdapter(this, SessionsInterface->OnCreateSessionCompleteDelegates,
+	[this, Promise = MoveTemp(Promise), Params](FName SessionName, const bool bWasSuccessful) mutable
+	{
+		if (!bWasSuccessful)
 		{
-			Op.SetError(MoveTemp(StateCheck));
+			Promise.EmplaceValue(Errors::Unknown());
 			return;
 		}
 
-		// Dedicated servers can update sessions so we won't check if LocalAccountId is valid
-		const FUniqueNetIdRef LocalAccountId = Services.Get<FAuthOSSAdapter>()->GetUniqueNetId(OpParams.LocalAccountId).ToSharedRef();
-
-		MakeMulticastAdapter(this, SessionsInterface->OnCreateSessionCompleteDelegates,
-		[this, WeakOp = Op.AsWeak()](FName SessionName, const bool bWasSuccessful) mutable
-		{
-			if (TOnlineAsyncOpPtr<FCreateSession> Op = WeakOp.Pin())
-			{
-				if (!bWasSuccessful)
-				{
-					Op->SetError(Errors::Unknown());
-					return;
-				}
-
-				const FCreateSession::Params& OpParams = Op->GetParams();
-
-				const FNamedOnlineSession* V1Session = SessionsInterface->GetNamedSession(SessionName);
-
-				TSharedRef<FSessionCommon> V2Session = BuildV2Session(V1Session);
+		const FNamedOnlineSession* V1Session = SessionsInterface->GetNamedSession(SessionName);
 
-				AddSessionWithReferences(V2Session, OpParams.SessionName, OpParams.LocalAccountId, OpParams.bPresenceEnabled);
+		TSharedRef<FSessionCommon> V2Session = BuildV2Session(V1Session);
 
-				Op->SetResult({ });
-			}
-		});
+		AddSessionWithReferences(V2Session, Params.SessionName, Params.LocalAccountId, Params.bPresenceEnabled);
 
-		SessionsInterface->CreateSession(*LocalAccountId, OpParams.SessionName, BuildV1SettingsForCreate(OpParams));
+		Promise.EmplaceValue(FCreateSession::Result{ });
+	});
 
-		// TODO: Pending support for multiple local users
-	})
-	.Enqueue(GetSerialQueue());
+	SessionsInterface->CreateSession(*LocalAccountId, Params.SessionName, BuildV1SettingsForCreate(Params));
 
-	return Op->GetHandle();
+	return Future;
 }
 
-TFuture<TOnlineResult<FUpdateSessionSettingsImpl>> FSessionsOSSAdapter::UpdateSessionSettingsImpl(FUpdateSessionSettingsImpl::Params&& Params)
+TFuture<TOnlineResult<FUpdateSessionSettings>> FSessionsOSSAdapter::UpdateSessionSettingsImpl(const FUpdateSessionSettings::Params& Params)
 {
-	TPromise<TOnlineResult<FUpdateSessionSettingsImpl>> Promise;
-	TFuture<TOnlineResult<FUpdateSessionSettingsImpl>> Future = Promise.GetFuture();
+	TPromise<TOnlineResult<FUpdateSessionSettings>> Promise;
+	TFuture<TOnlineResult<FUpdateSessionSettings>> Future = Promise.GetFuture();
 
 	MakeMulticastAdapter(this, SessionsInterface->OnUpdateSessionCompleteDelegates,
 	[this, Promise = MoveTemp(Promise), OpParams = Params](FName SessionName, const bool bWasSuccessful) mutable
@@ -476,7 +450,7 @@ TFuture<TOnlineResult<FUpdateSessionSettingsImpl>> FSessionsOSSAdapter::UpdateSe
 			(*FoundSession) += SessionUpdateData;
 
 			// We set the result and fire the event
-			Promise.EmplaceValue(FUpdateSessionSettingsImpl::Result{ });
+			Promise.EmplaceValue(FUpdateSessionSettings::Result{ });
 
 			FSessionUpdated SessionUpdatedEvent{ OpParams.SessionName, SessionUpdateData };
 			SessionEvents.OnSessionUpdated.Broadcast(SessionUpdatedEvent);
@@ -505,562 +479,413 @@ TFuture<TOnlineResult<FUpdateSessionSettingsImpl>> FSessionsOSSAdapter::UpdateSe
 	return Future;
 }
 
-TOnlineAsyncOpHandle<FLeaveSession> FSessionsOSSAdapter::LeaveSession(FLeaveSession::Params&& Params)
+TFuture<TOnlineResult<FLeaveSession>> FSessionsOSSAdapter::LeaveSessionImpl(const FLeaveSession::Params& Params)
 {
-	TOnlineAsyncOpRef<FLeaveSession> Op = GetOp<FLeaveSession>(MoveTemp(Params));
-	const FLeaveSession::Params& OpParams = Op->GetParams();
+	TPromise<TOnlineResult<FLeaveSession>> Promise;
+	TFuture<TOnlineResult<FLeaveSession>> Future = Promise.GetFuture();
 
-	Op->Then([this](TOnlineAsyncOp<FLeaveSession>& Op)
+	MakeMulticastAdapter(this, SessionsInterface->OnDestroySessionCompleteDelegates,
+	[this, Promise = MoveTemp(Promise), Params](FName SessionName, const bool bWasSuccessful) mutable
 	{
-		const FLeaveSession::Params& OpParams = Op.GetParams();
-
-		FOnlineError StateCheck = CheckLeaveSessionState(OpParams);
-		if (StateCheck != Errors::Success())
+		if (!bWasSuccessful)
 		{
-			Op.SetError(MoveTemp(StateCheck));
+			Promise.EmplaceValue(Errors::Unknown());
+			SessionsInterface->ClearOnDestroySessionCompleteDelegates(this);
 			return;
 		}
 
-		MakeMulticastAdapter(this, SessionsInterface->OnDestroySessionCompleteDelegates,
-		[this, WeakOp = Op.AsWeak()](FName SessionName, const bool bWasSuccessful) mutable
+		TOnlineResult<FGetSessionByName> GetSessionByNameResult = GetSessionByName({ Params.SessionName });
+		if (GetSessionByNameResult.IsOk())
 		{
-			if (TOnlineAsyncOpPtr<FLeaveSession> Op = WeakOp.Pin())
-			{
-				if (!bWasSuccessful)
-				{
-					Op->SetError(Errors::Unknown());
-					SessionsInterface->ClearOnDestroySessionCompleteDelegates(this);
-					return;
-				}
+			TSharedRef<const ISession> FoundSession = GetSessionByNameResult.GetOkValue().Session;
 
-				const FLeaveSession::Params& OpParams = Op->GetParams();
+			ClearSessionReferences(FoundSession->GetSessionId(), Params.SessionName, Params.LocalAccountId);
 
-				TOnlineResult<FGetSessionByName> GetSessionByNameResult = GetSessionByName({ OpParams.SessionName });
-				if (GetSessionByNameResult.IsOk())
-				{
-					TSharedRef<const ISession> FoundSession = GetSessionByNameResult.GetOkValue().Session;
+			Promise.EmplaceValue(FLeaveSession::Result{ });
 
-					ClearSessionReferences(FoundSession->GetSessionId(), OpParams.SessionName, OpParams.LocalAccountId);
-				}
+			FSessionLeft SessionLeftEvent;
+			SessionLeftEvent.LocalAccountId = Params.LocalAccountId;
 
-				Op->SetResult({ });
-
-				FSessionLeft SessionLeftEvent;
-				SessionLeftEvent.LocalAccountId = Op->GetParams().LocalAccountId;
-
-				SessionEvents.OnSessionLeft.Broadcast(SessionLeftEvent);
-			}
-
-			SessionsInterface->ClearOnDestroySessionCompleteDelegates(this);
-		});
-
-		SessionsInterface->DestroySession(OpParams.SessionName);
+			SessionEvents.OnSessionLeft.Broadcast(SessionLeftEvent);
+		}
+		else
+		{
+			Promise.EmplaceValue(GetSessionByNameResult.GetErrorValue());
+		}
+	});
 
-		// TODO: Pending support for multiple local users
-	})
-	.Enqueue(GetSerialQueue());
+	SessionsInterface->DestroySession(Params.SessionName);
 
-	return Op->GetHandle();
+	return Future;
 }
 
-TOnlineAsyncOpHandle<FFindSessions> FSessionsOSSAdapter::FindSessions(FFindSessions::Params&& Params)
+TFuture<TOnlineResult<FFindSessions>> FSessionsOSSAdapter::FindSessionsImpl(const FFindSessions::Params& Params)
 {
-	TOnlineAsyncOpRef<FFindSessions> Op = GetOp<FFindSessions>(MoveTemp(Params));
-	const FFindSessions::Params& OpParams = Op->GetParams();
+	TPromise<TOnlineResult<FFindSessions>> Promise;
+	TFuture<TOnlineResult<FFindSessions>> Future = Promise.GetFuture();
 
-	FOnlineError ParamsCheck = CheckFindSessionsParams(OpParams);
-	if (ParamsCheck != Errors::Success())
+	const FUniqueNetIdRef LocalAccountId = Services.Get<FAuthOSSAdapter>()->GetUniqueNetId(Params.LocalAccountId).ToSharedRef();
+	if (!LocalAccountId->IsValid())
 	{
-		Op->SetError(MoveTemp(ParamsCheck));
-		return Op->GetHandle();
+		Promise.EmplaceValue(Errors::InvalidUser());
+		return Future;
 	}
 
-	Op->Then([this](TOnlineAsyncOp<FFindSessions>& Op)
+	if (PendingV1SessionSearchesPerUser.Contains(Params.LocalAccountId))
 	{
-		const FFindSessions::Params& OpParams = Op.GetParams();
-
-		FOnlineError StateCheck = CheckFindSessionsState(OpParams);
-		if (StateCheck != Errors::Success())
-		{
-			Op.SetError(MoveTemp(StateCheck));
-			return;
-		}
+		Promise.EmplaceValue(Errors::AlreadyPending());
+		return Future;
+	}
 
-		const FUniqueNetIdRef LocalAccountId = Services.Get<FAuthOSSAdapter>()->GetUniqueNetId(OpParams.LocalAccountId).ToSharedRef();
-		if (!LocalAccountId->IsValid())
-		{
-			Op.SetError(Errors::InvalidUser());
-			return;
-		}
+	// Before we start the search, we reset the cache and save the promise
+	SearchResultsUserMap.FindOrAdd(Params.LocalAccountId).Reset();
+	CurrentSessionSearchPromisesUserMap.Emplace(Params.LocalAccountId, MoveTemp(Promise));
 
-		if (PendingV1SessionSearchesPerUser.Contains(OpParams.LocalAccountId))
+	if (Params.SessionId.IsSet())
+	{
+		if (const FUniqueNetIdPtr SessionIdPtr = GetSessionIdRegistry().GetIdValue(*Params.SessionId))
 		{
-			Op.SetError(Errors::AlreadyPending());
-			return;
-		}
+			const FUniqueNetIdRef SessionIdRef = SessionIdPtr.ToSharedRef();
 
-		// If there is a target user, we will use FindFriendSession
-		if (OpParams.SessionId.IsSet())
-		{
-			if (const FUniqueNetIdPtr SessionIdPtr = GetSessionIdRegistry().GetIdValue(*OpParams.SessionId))
+			MakeMulticastAdapter(this, SessionsInterface->OnFindSessionsCompleteDelegates,
+			[this, Params](bool bWasSuccessful)
 			{
-				const FUniqueNetIdRef SessionIdRef = SessionIdPtr.ToSharedRef();
+				TPromise<TOnlineResult<FFindSessions>>& Promise = CurrentSessionSearchPromisesUserMap.FindChecked(Params.LocalAccountId);
 
-				MakeMulticastAdapter(this, SessionsInterface->OnFindSessionsCompleteDelegates,
-				[this, WeakOp = Op.AsWeak()](bool bWasSuccessful)
+				if (bWasSuccessful)
 				{
-					if (TOnlineAsyncOpPtr<FFindSessions> Op = WeakOp.Pin())
+					TSharedRef<FOnlineSessionSearch>& PendingV1SessionSearch = PendingV1SessionSearchesPerUser.FindChecked(Params.LocalAccountId);
+
+					TArray<TSharedRef<FSessionCommon>> FoundSessions = BuildV2SessionSearchResults(PendingV1SessionSearch->SearchResults);
+					for (const TSharedRef<FSessionCommon>& FoundSession : FoundSessions)
 					{
-						if (bWasSuccessful)
-						{
-							TSharedRef<FOnlineSessionSearch>& PendingV1SessionSearch = PendingV1SessionSearchesPerUser.FindChecked(Op->GetParams().LocalAccountId);
-
-							TArray<TSharedRef<FSessionCommon>> FoundSessions = BuildV2SessionSearchResults(PendingV1SessionSearch->SearchResults);
-							for (const TSharedRef<FSessionCommon>& FoundSession : FoundSessions)
-							{
-								AddSearchResult(FoundSession, Op->GetParams().LocalAccountId);
-							}
-
-							TArray<FOnlineSessionId> SearchResults = SearchResultsUserMap.FindOrAdd(Op->GetParams().LocalAccountId);
-							Op->SetResult({ SearchResults });
-						}
-						else
-						{
-							Op->SetError(Errors::Unknown());
-						}
-
-						PendingV1SessionSearchesPerUser.Remove(Op->GetParams().LocalAccountId);
+						AddSearchResult(FoundSession, Params.LocalAccountId);
 					}
-				});
 
-				FOnlineServicesOSSAdapter& ServicesOSSAdapter = static_cast<FOnlineServicesOSSAdapter&>(Services);
-
-				FUniqueNetIdRef TargetIdRef = FUniqueNetIdString::EmptyId();
-				if (OpParams.TargetUser.IsSet())
+					TArray<FOnlineSessionId> SearchResults = SearchResultsUserMap.FindOrAdd(Params.LocalAccountId);
+					Promise.EmplaceValue(FFindSessions::Result{ SearchResults });
+				}
+				else
 				{
-					TargetIdRef = ServicesOSSAdapter.GetAccountIdRegistry().GetIdValue(*OpParams.TargetUser).ToSharedRef();
+					Promise.EmplaceValue(Errors::Unknown());
 				}
 
-				PendingV1SessionSearchesPerUser.Emplace(OpParams.LocalAccountId, MakeShared<FOnlineSessionSearch>());
-
-				SessionsInterface->FindSessionById(*LocalAccountId, *SessionIdRef, *TargetIdRef, *MakeDelegateAdapter(this, [this, WeakOp = Op.AsWeak()](int32 LocalUserNum, bool bWasSuccessful, const FOnlineSessionSearchResult& SingleSearchResult) mutable
-				{
-					if (TOnlineAsyncOpPtr<FFindSessions> Op = WeakOp.Pin())
-					{
-						if (!bWasSuccessful)
-						{
-							Op->SetError(Errors::Unknown());
-							return;
-						}
+				CurrentSessionSearchPromisesUserMap.Remove(Params.LocalAccountId);
+				PendingV1SessionSearchesPerUser.Remove(Params.LocalAccountId);
+			});
 
-						TSharedRef<FOnlineSessionSearch>& PendingV1SessionSearch = PendingV1SessionSearchesPerUser.FindChecked(Op->GetParams().LocalAccountId);
+			FOnlineServicesOSSAdapter& ServicesOSSAdapter = static_cast<FOnlineServicesOSSAdapter&>(Services);
 
-						PendingV1SessionSearch->SearchResults.Add(SingleSearchResult);
-					}
-				}));
-			}
-			else
+			FUniqueNetIdRef TargetIdRef = FUniqueNetIdString::EmptyId();
+			if (Params.TargetUser.IsSet())
 			{
-				Op.SetError(Errors::InvalidParams());
-				return;
+				TargetIdRef = ServicesOSSAdapter.GetAccountIdRegistry().GetIdValue(*Params.TargetUser).ToSharedRef();
 			}
-		}
-		else if (OpParams.TargetUser.IsSet())
-		{
-			int32 LocalUserNum = Services.Get<FAuthOSSAdapter>()->GetLocalUserNum(OpParams.LocalAccountId);
 
-			MakeMulticastAdapter(this, SessionsInterface->OnFindFriendSessionCompleteDelegates[LocalUserNum],
-			[this, WeakOp = Op.AsWeak()](int32 LocalAccountId, bool bWasSuccessful, const TArray<FOnlineSessionSearchResult>& FriendSearchResults)
+			PendingV1SessionSearchesPerUser.Emplace(Params.LocalAccountId, MakeShared<FOnlineSessionSearch>());
+
+			SessionsInterface->FindSessionById(*LocalAccountId, *SessionIdRef, *TargetIdRef, *MakeDelegateAdapter(this, [this, Params](int32 LocalUserNum, bool bWasSuccessful, const FOnlineSessionSearchResult& SingleSearchResult) mutable
 			{
-				if (TOnlineAsyncOpPtr<FFindSessions> Op = WeakOp.Pin())
-				{
-					if (bWasSuccessful)
-					{
-						TArray<TSharedRef<FSessionCommon>> FoundSessions = BuildV2SessionSearchResults(FriendSearchResults);
-						for (const TSharedRef<FSessionCommon>& FoundSession : FoundSessions)
-						{
-							AddSearchResult(FoundSession, Op->GetParams().LocalAccountId);
-						}
-
-						TArray<FOnlineSessionId> SearchResults = SearchResultsUserMap.FindOrAdd(Op->GetParams().LocalAccountId);
-						Op->SetResult({ SearchResults });
-					}
-					else
-					{
-						Op->SetError(Errors::Unknown());						
-					}
+				TPromise<TOnlineResult<FFindSessions>>& Promise = CurrentSessionSearchPromisesUserMap.FindChecked(Params.LocalAccountId);
 
-					PendingV1SessionSearchesPerUser.Remove(Op->GetParams().LocalAccountId);
+				if (!bWasSuccessful)
+				{
+					Promise.EmplaceValue(Errors::Unknown());
+					return;
 				}
-			});
 
-			FOnlineServicesOSSAdapter& ServicesOSSAdapter = static_cast<FOnlineServicesOSSAdapter&>(Services);
-
-			TArray<FUniqueNetIdRef> TargetUniqueNetIds = { ServicesOSSAdapter.GetAccountIdRegistry().GetIdValue(*OpParams.TargetUser).ToSharedRef() };
+				TSharedRef<FOnlineSessionSearch>& PendingV1SessionSearch = PendingV1SessionSearchesPerUser.FindChecked(Params.LocalAccountId);
 
-			// The cached search information is not used in this mode, but we'll still update it to not have more than one search at the same time
-			PendingV1SessionSearchesPerUser.Emplace(OpParams.LocalAccountId, MakeShared<FOnlineSessionSearch>());
-
-			SessionsInterface->FindFriendSession(*LocalAccountId, TargetUniqueNetIds);
+				PendingV1SessionSearch->SearchResults.Add(SingleSearchResult);
+			}));
 		}
 		else
 		{
-			MakeMulticastAdapter(this, SessionsInterface->OnFindSessionsCompleteDelegates,
-			[this, WeakOp = Op.AsWeak()](bool bWasSuccessful) mutable
+			Promise.EmplaceValue(Errors::InvalidParams());
+			return Future;
+		}
+	}
+	else if (Params.TargetUser.IsSet()) 	// If there is a target user, we will use FindFriendSession
+	{
+		int32 LocalUserNum = Services.Get<FAuthOSSAdapter>()->GetLocalUserNum(Params.LocalAccountId);
+
+		MakeMulticastAdapter(this, SessionsInterface->OnFindFriendSessionCompleteDelegates[LocalUserNum],
+		[this, Params](int32 LocalAccountId, bool bWasSuccessful, const TArray<FOnlineSessionSearchResult>& FriendSearchResults)
+		{
+			TPromise<TOnlineResult<FFindSessions>>& Promise = CurrentSessionSearchPromisesUserMap.FindChecked(Params.LocalAccountId);
+
+			if (bWasSuccessful)
 			{
-				if (TOnlineAsyncOpPtr<FFindSessions> Op = WeakOp.Pin())
+				TArray<TSharedRef<FSessionCommon>> FoundSessions = BuildV2SessionSearchResults(FriendSearchResults);
+				for (const TSharedRef<FSessionCommon>& FoundSession : FoundSessions)
 				{
-					if (bWasSuccessful)
-					{
-						TSharedRef<FOnlineSessionSearch>& PendingV1SessionSearch = PendingV1SessionSearchesPerUser.FindChecked(Op->GetParams().LocalAccountId);
+					AddSearchResult(FoundSession,Params.LocalAccountId);
+				}
 
-						TArray<TSharedRef<FSessionCommon>> FoundSessions = BuildV2SessionSearchResults(PendingV1SessionSearch->SearchResults);
-						for (const TSharedRef<FSessionCommon>& FoundSession : FoundSessions)
-						{
-							AddSearchResult(FoundSession, Op->GetParams().LocalAccountId);
-						}
+				TArray<FOnlineSessionId> SearchResults = SearchResultsUserMap.FindOrAdd(Params.LocalAccountId);
+				Promise.EmplaceValue(FFindSessions::Result{ SearchResults });
+			}
+			else
+			{
+				Promise.EmplaceValue(Errors::Unknown());
+			}
 
-						TArray<FOnlineSessionId> SearchResults = SearchResultsUserMap.FindOrAdd(Op->GetParams().LocalAccountId);
-						Op->SetResult({ SearchResults });
-					}
-					else
-					{
-						Op->SetError(Errors::Unknown());						
-					}
+			CurrentSessionSearchPromisesUserMap.Remove(Params.LocalAccountId);
+			PendingV1SessionSearchesPerUser.Remove(Params.LocalAccountId);
+		});
 
-					PendingV1SessionSearchesPerUser.Remove(Op->GetParams().LocalAccountId);
-				}
-			});
+		FOnlineServicesOSSAdapter& ServicesOSSAdapter = static_cast<FOnlineServicesOSSAdapter&>(Services);
 
-			TSharedRef<FOnlineSessionSearch>& PendingV1SessionSearch = PendingV1SessionSearchesPerUser.Emplace(OpParams.LocalAccountId, BuildV1SessionSearch(OpParams));
+		TArray<FUniqueNetIdRef> TargetUniqueNetIds = { ServicesOSSAdapter.GetAccountIdRegistry().GetIdValue(*Params.TargetUser).ToSharedRef() };
 
-			SessionsInterface->FindSessions(*LocalAccountId, PendingV1SessionSearch);
-		}
-	})
-	.Enqueue(GetSerialQueue());
+		// The cached search information is not used in this mode, but we'll still update it to not have more than one search at the same time
+		PendingV1SessionSearchesPerUser.Emplace(Params.LocalAccountId, MakeShared<FOnlineSessionSearch>());
 
-	return Op->GetHandle();
-}
+		SessionsInterface->FindFriendSession(*LocalAccountId, TargetUniqueNetIds);
+	}
+	else
+	{
+		MakeMulticastAdapter(this, SessionsInterface->OnFindSessionsCompleteDelegates,
+		[this, Params](bool bWasSuccessful) mutable
+		{
+			TPromise<TOnlineResult<FFindSessions>>& Promise = CurrentSessionSearchPromisesUserMap.FindChecked(Params.LocalAccountId);
 
-TOnlineAsyncOpHandle<FStartMatchmaking> FSessionsOSSAdapter::StartMatchmaking(FStartMatchmaking::Params&& Params)
-{
-	TOnlineAsyncOpRef<FStartMatchmaking> Op = GetOp<FStartMatchmaking>(MoveTemp(Params));
-	const FStartMatchmaking::Params& OpParams = Op->GetParams();
+			if (bWasSuccessful)
+			{
+				TSharedRef<FOnlineSessionSearch>& PendingV1SessionSearch = PendingV1SessionSearchesPerUser.FindChecked(Params.LocalAccountId);
 
-	FOnlineError ParamsCheck = CheckStartMatchmakingParams(OpParams);
-	if (ParamsCheck != Errors::Success())
-	{
-		Op->SetError(MoveTemp(ParamsCheck));
-		return Op->GetHandle();
-	}
+				TArray<TSharedRef<FSessionCommon>> FoundSessions = BuildV2SessionSearchResults(PendingV1SessionSearch->SearchResults);
+				for (const TSharedRef<FSessionCommon>& FoundSession : FoundSessions)
+				{
+					AddSearchResult(FoundSession, Params.LocalAccountId);
+				}
 
-	Op->Then([this](TOnlineAsyncOp<FStartMatchmaking>& Op)
-	{
-		const FStartMatchmaking::Params& OpParams = Op.GetParams();
+				TArray<FOnlineSessionId> SearchResults = SearchResultsUserMap.FindOrAdd(Params.LocalAccountId);
+				Promise.EmplaceValue(FFindSessions::Result{ SearchResults });
+			}
+			else
+			{
+				Promise.EmplaceValue(Errors::Unknown());
+			}
 
-		FOnlineError StateCheck = CheckStartMatchmakingState(OpParams);
-		if (StateCheck != Errors::Success())
-		{
-			Op.SetError(MoveTemp(StateCheck));
-			return;
-		}
+			CurrentSessionSearchPromisesUserMap.Remove(Params.LocalAccountId);
+			PendingV1SessionSearchesPerUser.Remove(Params.LocalAccountId);
+		});
 
-		if (PendingV1SessionSearchesPerUser.Contains(OpParams.SessionCreationParameters.LocalAccountId))
-		{
-			Op.SetError(Errors::AlreadyPending());
-			return;
-		}
+		TSharedRef<FOnlineSessionSearch>& PendingV1SessionSearch = PendingV1SessionSearchesPerUser.Emplace(Params.LocalAccountId, BuildV1SessionSearch(Params));
 
-		FOnlineServicesOSSAdapter& ServicesOSSAdapter = static_cast<FOnlineServicesOSSAdapter&>(Services);
+		SessionsInterface->FindSessions(*LocalAccountId, PendingV1SessionSearch);
+	}
+
+	return Future;
+}
 
-		TArray<FSessionMatchmakingUser> MatchMakingUsers;
-		FSessionMatchmakingUser NewMatchMakingUser { ServicesOSSAdapter.GetAccountIdRegistry().GetIdValue(OpParams.SessionCreationParameters.LocalAccountId).ToSharedRef() };
-		MatchMakingUsers.Add(NewMatchMakingUser);
+TFuture<TOnlineResult<FStartMatchmaking>> FSessionsOSSAdapter::StartMatchmakingImpl(const FStartMatchmaking::Params& Params)
+{
+	TPromise<TOnlineResult<FStartMatchmaking>> Promise;
+	TFuture<TOnlineResult<FStartMatchmaking>> Future = Promise.GetFuture();
 
-		FOnlineSessionSettings NewSessionSettings = BuildV1SettingsForCreate(OpParams.SessionCreationParameters);
+	FOnlineServicesOSSAdapter& ServicesOSSAdapter = static_cast<FOnlineServicesOSSAdapter&>(Services);
 
-		/** We build a mock params struct for FindSessions and convert it to a V1 session search */
-		FFindSessions::Params FindSessionsParams;
-		FindSessionsParams.bFindLANSessions = false;
-		FindSessionsParams.Filters = OpParams.SessionSearchFilters;
-		FindSessionsParams.LocalAccountId = OpParams.SessionCreationParameters.LocalAccountId;
-		FindSessionsParams.MaxResults = 1;
+	TArray<FSessionMatchmakingUser> MatchMakingUsers;
+	FSessionMatchmakingUser NewMatchMakingUser { ServicesOSSAdapter.GetAccountIdRegistry().GetIdValue(Params.SessionCreationParameters.LocalAccountId).ToSharedRef() };
+	MatchMakingUsers.Add(NewMatchMakingUser);
 
-		TSharedRef<FOnlineSessionSearch>& SessionSearch = PendingV1SessionSearchesPerUser.Emplace(OpParams.SessionCreationParameters.LocalAccountId, BuildV1SessionSearch(FindSessionsParams));
+	FOnlineSessionSettings NewSessionSettings = BuildV1SettingsForCreate(Params.SessionCreationParameters);
 
-		// QUESTION: StartMatchmaking causes both OnMatchmakingComplete and OnStartMatchmakingComplete to trigger.
-		// The first triggers the second so I used the latter, but the interface header did not specify if one should be used. Is this the correct one?
-		SessionsInterface->StartMatchmaking(MatchMakingUsers, OpParams.SessionCreationParameters.SessionName, NewSessionSettings, SessionSearch, *MakeDelegateAdapter(this, [this, WeakOp = Op.AsWeak()](FName SessionName, const ::FOnlineError& ErrorDetails, const FSessionMatchmakingResults& Results) mutable
-		{
-			if (TOnlineAsyncOpPtr<FStartMatchmaking> Op = WeakOp.Pin())
-			{
-				if (!ErrorDetails.bSucceeded)
-				{
-					Op->SetError(Errors::Unknown()); // TODO: Proper error parsing
-					return;
-				}
+	/** We build a mock params struct for FindSessions and convert it to a V1 session search */
+	FFindSessions::Params FindSessionsParams;
+	FindSessionsParams.bFindLANSessions = false;
+	FindSessionsParams.Filters = Params.SessionSearchFilters;
+	FindSessionsParams.LocalAccountId = Params.SessionCreationParameters.LocalAccountId;
+	FindSessionsParams.MaxResults = 1;
 
-				const FStartMatchmaking::Params& OpParams = Op->GetParams();
+	TSharedRef<FOnlineSessionSearch>& SessionSearch = PendingV1SessionSearchesPerUser.Emplace(Params.SessionCreationParameters.LocalAccountId, BuildV1SessionSearch(FindSessionsParams));
 
-				TSharedRef<FSessionCommon> V2Session = BuildV2Session(SessionsInterface->GetNamedSession(SessionName));
+	SessionsInterface->StartMatchmaking(MatchMakingUsers, Params.SessionCreationParameters.SessionName, NewSessionSettings, SessionSearch, *MakeDelegateAdapter(this, [this, Promise = MoveTemp(Promise), Params](FName SessionName, const ::FOnlineError& ErrorDetails, const FSessionMatchmakingResults& Results) mutable
+	{
+		if (!ErrorDetails.bSucceeded)
+		{
+			Promise.EmplaceValue(Errors::Unknown()); // TODO: Proper error parsing
+			return;
+		}
 
-				AddSessionWithReferences(V2Session, OpParams.SessionCreationParameters.SessionName, OpParams.SessionCreationParameters.LocalAccountId, OpParams.SessionCreationParameters.bPresenceEnabled);
+		TSharedRef<FSessionCommon> V2Session = BuildV2Session(SessionsInterface->GetNamedSession(SessionName));
 
-				Op->SetResult({ });
+		AddSessionWithReferences(V2Session, Params.SessionCreationParameters.SessionName, Params.SessionCreationParameters.LocalAccountId, Params.SessionCreationParameters.bPresenceEnabled);
 
-				FSessionJoined Event{ { Op->GetParams().SessionCreationParameters.LocalAccountId }, V2Session->GetSessionId()};
+		Promise.EmplaceValue(FStartMatchmaking::Result{ });
 
-				SessionEvents.OnSessionJoined.Broadcast(Event);
+		FSessionJoined Event{ { Params.SessionCreationParameters.LocalAccountId }, V2Session->GetSessionId()};
 
-				// TODO: Pending multiple local user support
-			}
-		}));
-	})
-	.Enqueue(GetSerialQueue());
+		SessionEvents.OnSessionJoined.Broadcast(Event);
+	}));
 
-	return Op->GetHandle();
+	return Future;
 }
 
-TOnlineAsyncOpHandle<FJoinSession> FSessionsOSSAdapter::JoinSession(FJoinSession::Params&& Params)
+TFuture<TOnlineResult<FJoinSession>> FSessionsOSSAdapter::JoinSessionImpl(const FJoinSession::Params& Params)
 {
-	TOnlineAsyncOpRef<FJoinSession> Op = GetOp<FJoinSession>(MoveTemp(Params));
-	const FJoinSession::Params& OpParams = Op->GetParams();
+	TPromise<TOnlineResult<FJoinSession>> Promise;
+	TFuture<TOnlineResult<FJoinSession>> Future = Promise.GetFuture();
 
-	FOnlineError ParamsCheck = CheckJoinSessionParams(OpParams);
-	if (ParamsCheck != Errors::Success())
+	const FUniqueNetIdRef LocalAccountId = Services.Get<FAuthOSSAdapter>()->GetUniqueNetId(Params.LocalAccountId).ToSharedRef();
+	if (!LocalAccountId->IsValid())
 	{
-		Op->SetError(MoveTemp(ParamsCheck));
-		return Op->GetHandle();
+		Promise.EmplaceValue(Errors::InvalidUser());
+		return Future;
 	}
 
-	Op->Then([this](TOnlineAsyncOp<FJoinSession>& Op)
+	MakeMulticastAdapter(this, SessionsInterface->OnJoinSessionCompleteDelegates,
+	[this, Promise = MoveTemp(Promise), Params](FName SessionName, EOnJoinSessionCompleteResult::Type Result) mutable
 	{
-		const FJoinSession::Params& OpParams = Op.GetParams();
-
-		FOnlineError StateCheck = CheckJoinSessionState(OpParams);
-		if (StateCheck != Errors::Success())
+		if (Result != EOnJoinSessionCompleteResult::Success)
 		{
-			Op.SetError(MoveTemp(StateCheck));
-			return;
-		}
-
-		const FUniqueNetIdRef LocalAccountId = Services.Get<FAuthOSSAdapter>()->GetUniqueNetId(OpParams.LocalAccountId).ToSharedRef();
-		if (!LocalAccountId->IsValid())
-		{
-			Op.SetError(Errors::InvalidUser());
+			Promise.EmplaceValue(Errors::Unknown());
 			return;
 		}
 
-		MakeMulticastAdapter(this, SessionsInterface->OnJoinSessionCompleteDelegates,
-		[this, WeakOp = Op.AsWeak()](FName SessionName, EOnJoinSessionCompleteResult::Type Result) mutable
+		TOnlineResult<FGetSessionById> GetSessionByIdResult = GetSessionById({ Params.SessionId });
+		if (GetSessionByIdResult.IsError())
 		{
-			if (TOnlineAsyncOpPtr<FJoinSession> Op = WeakOp.Pin())
+			// If no result is found, the id might be expired, which we should notify
+			if (GetSessionIdRegistry().IsHandleExpired(Params.SessionId))
 			{
-				if (Result != EOnJoinSessionCompleteResult::Success)
-				{
-					Op->SetError(Errors::Unknown());
-					return;
-				}
-
-				const FJoinSession::Params& OpParams = Op->GetParams();
-
-				TOnlineResult<FGetSessionById> GetSessionByIdResult = GetSessionById({ OpParams.SessionId });
-				if (GetSessionByIdResult.IsError())
-				{
-					// If no result is found, the id might be expired, which we should notify
-					if (GetSessionIdRegistry().IsHandleExpired(OpParams.SessionId))
-					{
-						UE_LOG(LogTemp, Warning, TEXT("[FSessionsOSSAdapter::JoinSession] SessionId parameter [%s] is expired. Please call FindSessions to get an updated list of available sessions "), *ToLogString(OpParams.SessionId));
-					}
-
-					Op->SetError(MoveTemp(GetSessionByIdResult.GetErrorValue()));
-					return;
-				}
+				UE_LOG(LogTemp, Warning, TEXT("[%s] SessionId parameter [%s] is expired. Please call FindSessions to get an updated list of available sessions "), UTF8_TO_TCHAR(__FUNCTION__), *ToLogString(Params.SessionId));
+			}
 
-				TSharedRef<const ISession> FoundSession = GetSessionByIdResult.GetOkValue().Session;
+			Promise.EmplaceValue(MoveTemp(GetSessionByIdResult.GetErrorValue()));
+			return;
+		}
 
-				AddSessionReferences(FoundSession->GetSessionId(), OpParams.SessionName, OpParams.LocalAccountId, OpParams.bPresenceEnabled);
+		TSharedRef<const ISession> FoundSession = GetSessionByIdResult.GetOkValue().Session;
 
-				// After successfully joining a session, we'll remove all related invites if any are found
-				ClearSessionInvitesForSession(OpParams.LocalAccountId, FoundSession->GetSessionId());
+		AddSessionReferences(FoundSession->GetSessionId(), Params.SessionName, Params.LocalAccountId, Params.bPresenceEnabled);
 
-				Op->SetResult({ });
+		// After successfully joining a session, we'll remove all related invites if any are found
+		ClearSessionInvitesForSession(Params.LocalAccountId, FoundSession->GetSessionId());
 
-				FSessionJoined Event{ { OpParams.LocalAccountId }, FoundSession->GetSessionId() };
+		Promise.EmplaceValue(FJoinSession::Result{ });
 
-				SessionEvents.OnSessionJoined.Broadcast(Event);
-			}
-		});
+		FSessionJoined Event{ { Params.LocalAccountId }, FoundSession->GetSessionId() };
 
-		TOnlineResult<FGetSessionById> GetSessionByIdResult = GetSessionById({ OpParams.SessionId });
-		if (GetSessionByIdResult.IsError())
-		{
-			Op.SetError(MoveTemp(GetSessionByIdResult.GetErrorValue()));
-			return;
-		}
+		SessionEvents.OnSessionJoined.Broadcast(Event);
+	});
 
-		const TSharedRef<const ISession>& FoundSession = GetSessionByIdResult.GetOkValue().Session;
+	TOnlineResult<FGetSessionById> GetSessionByIdResult = GetSessionById({ Params.SessionId });
+	if (GetSessionByIdResult.IsError())
+	{
+		Promise.EmplaceValue(GetSessionByIdResult.GetErrorValue());
+		return Future;
+	}
 
-		FOnlineSessionSearchResult SearchResult;
+	const TSharedRef<const ISession>& FoundSession = GetSessionByIdResult.GetOkValue().Session;
 
-		if (const FCustomSessionSetting* PingInMs = FoundSession->GetSessionSettings().CustomSettings.Find(OSS_ADAPTER_SESSIONS_PING_IN_MS))
-		{
-			SearchResult.PingInMs = (int32)PingInMs->Data.GetInt64();
-		}
+	FOnlineSessionSearchResult SearchResult;
 
-		SearchResult.Session = BuildV1Session(FoundSession);
+	if (const FCustomSessionSetting* PingInMs = FoundSession->GetSessionSettings().CustomSettings.Find(OSS_ADAPTER_SESSIONS_PING_IN_MS))
+	{
+		SearchResult.PingInMs = (int32)PingInMs->Data.GetInt64();
+	}
 
-		SessionsInterface->JoinSession(*LocalAccountId, OpParams.SessionName, SearchResult);
+	SearchResult.Session = BuildV1Session(FoundSession);
 
-		// TODO: Pending support for multiple local users
-	})
-	.Enqueue(GetSerialQueue());
+	SessionsInterface->JoinSession(*LocalAccountId, Params.SessionName, SearchResult);
 
-	return Op->GetHandle();
+	return Future;
 }
 
-TOnlineAsyncOpHandle<FAddSessionMember> FSessionsOSSAdapter::AddSessionMember(FAddSessionMember::Params&& Params)
+TFuture<TOnlineResult<FAddSessionMember>> FSessionsOSSAdapter::AddSessionMemberImpl(const FAddSessionMember::Params& Params)
 {
-	TOnlineAsyncOpRef<FAddSessionMember> Op = GetOp<FAddSessionMember>(MoveTemp(Params));
-	const FAddSessionMember::Params& OpParams = Op->GetParams();
-
-	Op->Then([this](TOnlineAsyncOp<FAddSessionMember>& Op)
-	{
-		const FAddSessionMember::Params& OpParams = Op.GetParams();
+	TPromise<TOnlineResult<FAddSessionMember>> Promise;
+	TFuture<TOnlineResult<FAddSessionMember>> Future = Promise.GetFuture();
 
-		FOnlineError StateCheck = CheckAddSessionMemberState(OpParams);
-		if (StateCheck != Errors::Success())
+	MakeMulticastAdapter(this, SessionsInterface->OnRegisterPlayersCompleteDelegates,
+		[this, Promise = MoveTemp(Promise), Params](FName SessionName, const TArray<FUniqueNetIdRef>& Players, bool bWasSuccessful) mutable
 		{
-			Op.SetError(MoveTemp(StateCheck));
-			return;
-		}
-
-		SessionsInterface->AddOnRegisterPlayersCompleteDelegate_Handle(*MakeDelegateAdapter(this, [this, WeakOp = Op.AsWeak()](FName SessionName, const TArray<FUniqueNetIdRef>& Players, bool bWasSuccessful) mutable
-		{
-			// Clearing the delegate needs to be the last thing we do, as modifying captured variables after it would result in a crash
-
-			if (TOnlineAsyncOpPtr<FAddSessionMember> Op = WeakOp.Pin())
+			if (bWasSuccessful)
 			{
-				if (bWasSuccessful)
-				{
-					Op->SetResult({ });
-				}
-				else
-				{
-					Op->SetError(Errors::Unknown());
-				}
+				Promise.EmplaceValue(FAddSessionMember::Result{ });
 			}
+			else
+			{
+				Promise.EmplaceValue(Errors::Unknown());
+			}
+		});
 
-			SessionsInterface->ClearOnRegisterPlayersCompleteDelegates(this);
-		}));
-
-		FAuthOSSAdapter* Auth = Services.Get<FAuthOSSAdapter>();
-		TArray<FUniqueNetIdRef> TargetUserNetIds;		
-		FUniqueNetIdRef TargetUserNetId = Auth->GetUniqueNetId(OpParams.LocalAccountId).ToSharedRef();
-		if (TargetUserNetId->IsValid())
-		{
-			TargetUserNetIds.Add(TargetUserNetId);
-		}	
+	FAuthOSSAdapter* Auth = Services.Get<FAuthOSSAdapter>();
+	TArray<FUniqueNetIdRef> TargetUserNetIds;		
+	FUniqueNetIdRef TargetUserNetId = Auth->GetUniqueNetId(Params.LocalAccountId).ToSharedRef();
+	if (TargetUserNetId->IsValid())
+	{
+		TargetUserNetIds.Add(TargetUserNetId);
+	}	
 
-		SessionsInterface->RegisterPlayers(OpParams.SessionName, TargetUserNetIds);
-	})
-	.Enqueue(GetSerialQueue());
+	SessionsInterface->RegisterPlayers(Params.SessionName, TargetUserNetIds);
 
-	return Op->GetHandle();
+	return Future;
 }
 
-TOnlineAsyncOpHandle<FRemoveSessionMember> FSessionsOSSAdapter::RemoveSessionMember(FRemoveSessionMember::Params&& Params)
+TFuture<TOnlineResult<FRemoveSessionMember>> FSessionsOSSAdapter::RemoveSessionMemberImpl(const FRemoveSessionMember::Params& Params)
 {
-	TOnlineAsyncOpRef<FRemoveSessionMember> Op = GetOp<FRemoveSessionMember>(MoveTemp(Params));
-	const FRemoveSessionMember::Params& OpParams = Op->GetParams();
-
-	Op->Then([this](TOnlineAsyncOp<FRemoveSessionMember>& Op)
-	{
-		const FRemoveSessionMember::Params& OpParams = Op.GetParams();
+	TPromise<TOnlineResult<FRemoveSessionMember>> Promise;
+	TFuture<TOnlineResult<FRemoveSessionMember>> Future = Promise.GetFuture();
 
-		FOnlineError StateCheck = CheckRemoveSessionMemberState(OpParams);
-		if (StateCheck != Errors::Success())
+	MakeMulticastAdapter(this, SessionsInterface->OnUnregisterPlayersCompleteDelegates,
+		[this, Promise = MoveTemp(Promise)](FName SessionName, const TArray<FUniqueNetIdRef>& Players, bool bWasSuccessful) mutable
 		{
-			Op.SetError(MoveTemp(StateCheck));
-			return;
-		}
-
-		SessionsInterface->AddOnUnregisterPlayersCompleteDelegate_Handle(*MakeDelegateAdapter(this, [this, WeakOp = Op.AsWeak()](FName SessionName, const TArray<FUniqueNetIdRef>& Players, bool bWasSuccessful) mutable
-		{
-			// Clearing the delegate needs to be the last thing we do, as modifying captured variables after it would result in a crash
-
-			if (TOnlineAsyncOpPtr<FRemoveSessionMember> Op = WeakOp.Pin())
+			if (bWasSuccessful)
 			{
-				if (bWasSuccessful)
-				{
-					Op->SetResult({ });
-				}
-				else
-				{
-					Op->SetError(Errors::Unknown());
-				}
+				Promise.EmplaceValue(FRemoveSessionMember::Result{ });
 			}
+			else
+			{
+				Promise.EmplaceValue(Errors::Unknown());
+			}
+		});
 
-			SessionsInterface->ClearOnUnregisterPlayersCompleteDelegates(this);
-		}));
-
-		FAuthOSSAdapter* Auth = Services.Get<FAuthOSSAdapter>();
-		TArray<FUniqueNetIdRef> TargetUserNetIds;
-		FUniqueNetIdRef TargetUserNetId = Auth->GetUniqueNetId(OpParams.LocalAccountId).ToSharedRef();
-		if (TargetUserNetId->IsValid())
-		{
-			TargetUserNetIds.Add(TargetUserNetId);
-		}
+	FAuthOSSAdapter* Auth = Services.Get<FAuthOSSAdapter>();
+	TArray<FUniqueNetIdRef> TargetUserNetIds;
+	FUniqueNetIdRef TargetUserNetId = Auth->GetUniqueNetId(Params.LocalAccountId).ToSharedRef();
+	if (TargetUserNetId->IsValid())
+	{
+		TargetUserNetIds.Add(TargetUserNetId);
+	}
 
-		SessionsInterface->UnregisterPlayers(OpParams.SessionName, TargetUserNetIds);
-	})
-	.Enqueue(GetSerialQueue());
+	SessionsInterface->UnregisterPlayers(Params.SessionName, TargetUserNetIds);
 
-	return Op->GetHandle();
+	return Future;
 }
 
-TOnlineAsyncOpHandle<FSendSessionInvite> FSessionsOSSAdapter::SendSessionInvite(FSendSessionInvite::Params&& Params)
+TFuture<TOnlineResult<FSendSessionInvite>> FSessionsOSSAdapter::SendSessionInviteImpl(const FSendSessionInvite::Params& Params)
 {
-	TOnlineAsyncOpRef<FSendSessionInvite> Op = GetOp<FSendSessionInvite>(MoveTemp(Params));
-	const FSendSessionInvite::Params& OpParams = Op->GetParams();
+	TPromise<TOnlineResult<FSendSessionInvite>> Promise;
+	TFuture<TOnlineResult<FSendSessionInvite>> Future = Promise.GetFuture();
 
-	Op->Then([this](TOnlineAsyncOp<FSendSessionInvite>& Op)
-	{
-		const FSendSessionInvite::Params& OpParams = Op.GetParams();
+	FAuthOSSAdapter* Auth = Services.Get<FAuthOSSAdapter>();
 
-		FOnlineError StateCheck = CheckSendSessionInviteState(OpParams);
-		if (StateCheck != Errors::Success())
-		{
-			Op.SetError(MoveTemp(StateCheck));
-			return;
-		}
-
-		FAuthOSSAdapter* Auth = Services.Get<FAuthOSSAdapter>();
-
-		const FUniqueNetIdRef LocalAccountId = Auth->GetUniqueNetId(OpParams.LocalAccountId).ToSharedRef();
-		if (!LocalAccountId->IsValid())
-		{
-			Op.SetError(Errors::InvalidUser());
-			return;
-		}
+	const FUniqueNetIdRef LocalAccountId = Auth->GetUniqueNetId(Params.LocalAccountId).ToSharedRef();
+	if (!LocalAccountId->IsValid())
+	{
+		Promise.EmplaceValue(Errors::InvalidUser());
+		return Future;
+	}
 
-		TArray<FUniqueNetIdRef> TargetUserNetIds;
-		for (const FAccountId& TargetUser : OpParams.TargetUsers)
+	TArray<FUniqueNetIdRef> TargetUserNetIds;
+	for (const FAccountId& TargetUser : Params.TargetUsers)
+	{
+		FUniqueNetIdRef TargetUserNetId = Auth->GetUniqueNetId(TargetUser).ToSharedRef();
+		if (TargetUserNetId->IsValid())
 		{
-			FUniqueNetIdRef TargetUserNetId = Auth->GetUniqueNetId(TargetUser).ToSharedRef();
-			if (TargetUserNetId->IsValid())
-			{
-				TargetUserNetIds.Add(TargetUserNetId);
-			}
+			TargetUserNetIds.Add(TargetUserNetId);
 		}
+	}
 
-		SessionsInterface->SendSessionInviteToFriends(*LocalAccountId, OpParams.SessionName, TargetUserNetIds);
+	SessionsInterface->SendSessionInviteToFriends(*LocalAccountId, Params.SessionName, TargetUserNetIds);
 
-		Op.SetResult({ });
-	})
-	.Enqueue(GetSerialQueue());
+	Promise.EmplaceValue(FSendSessionInvite::Result{ });
 
-	return Op->GetHandle();
+	return Future;
 }
 
 TOnlineResult<FGetResolvedConnectString> FSessionsOSSAdapter::GetResolvedConnectString(const FGetResolvedConnectString::Params& Params)
diff --git a/Engine/Plugins/Online/Experimental/OnlineServicesOSSAdapter/Source/Public/Online/SessionsOSSAdapter.h b/Engine/Plugins/Online/Experimental/OnlineServicesOSSAdapter/Source/Public/Online/SessionsOSSAdapter.h
--- a/Engine/Plugins/Online/Experimental/OnlineServicesOSSAdapter/Source/Public/Online/SessionsOSSAdapter.h
+++ b/Engine/Plugins/Online/Experimental/OnlineServicesOSSAdapter/Source/Public/Online/SessionsOSSAdapter.h
@@ -64,18 +64,16 @@ public:
 	virtual void Initialize() override;
 	virtual void Shutdown() override;
 
-	// ISessions
-	virtual TOnlineAsyncOpHandle<FCreateSession> CreateSession(FCreateSession::Params&& Params) override;
-	virtual TOnlineAsyncOpHandle<FLeaveSession> LeaveSession(FLeaveSession::Params&& Params) override;
-	virtual TOnlineAsyncOpHandle<FFindSessions> FindSessions(FFindSessions::Params&& Params) override;
-	virtual TOnlineAsyncOpHandle<FStartMatchmaking> StartMatchmaking(FStartMatchmaking::Params&& Params) override;
-	virtual TOnlineAsyncOpHandle<FJoinSession> JoinSession(FJoinSession::Params&& Params) override;
-	virtual TOnlineAsyncOpHandle<FAddSessionMember> AddSessionMember(FAddSessionMember::Params&& Params) override;
-	virtual TOnlineAsyncOpHandle<FRemoveSessionMember> RemoveSessionMember(FRemoveSessionMember::Params&& Params) override;
-	virtual TOnlineAsyncOpHandle<FSendSessionInvite> SendSessionInvite(FSendSessionInvite::Params&& Params) override;
-
 	// FSessionsCommon
-	virtual TFuture<TOnlineResult<FUpdateSessionSettingsImpl>> UpdateSessionSettingsImpl(FUpdateSessionSettingsImpl::Params&& Params) override;
+	virtual TFuture<TOnlineResult<FCreateSession>> CreateSessionImpl(const FCreateSession::Params& Params) override;
+	virtual TFuture<TOnlineResult<FUpdateSessionSettings>> UpdateSessionSettingsImpl(const FUpdateSessionSettings::Params& Params) override;
+	virtual TFuture<TOnlineResult<FLeaveSession>> LeaveSessionImpl(const FLeaveSession::Params& Params) override;
+	virtual TFuture<TOnlineResult<FFindSessions>> FindSessionsImpl(const FFindSessions::Params& Params) override;
+	virtual TFuture<TOnlineResult<FStartMatchmaking>> StartMatchmakingImpl(const FStartMatchmaking::Params& Params) override;
+	virtual TFuture<TOnlineResult<FJoinSession>> JoinSessionImpl(const FJoinSession::Params& Params) override;
+	virtual TFuture<TOnlineResult<FAddSessionMember>> AddSessionMemberImpl(const FAddSessionMember::Params& Params) override;
+	virtual TFuture<TOnlineResult<FRemoveSessionMember>> RemoveSessionMemberImpl(const FRemoveSessionMember::Params& Params) override;
+	virtual TFuture<TOnlineResult<FSendSessionInvite>> SendSessionInviteImpl(const FSendSessionInvite::Params& Params) override;
 
 	TOnlineResult<FGetResolvedConnectString> GetResolvedConnectString(const FGetResolvedConnectString::Params& Params);
 
