Removal of bIsDedicatedServer session creation parameter as it should be automatically set

diff --git a/Engine/Plugins/Online/Experimental/OnlineServices/Source/OnlineServicesInterface/Public/Online/Sessions.h b/Engine/Plugins/Online/Experimental/OnlineServices/Source/OnlineServicesInterface/Public/Online/Sessions.h
--- a/Engine/Plugins/Online/Experimental/OnlineServices/Source/OnlineServicesInterface/Public/Online/Sessions.h
+++ b/Engine/Plugins/Online/Experimental/OnlineServices/Source/OnlineServicesInterface/Public/Online/Sessions.h
@@ -353,9 +353,6 @@ struct FCreateSession
 		/** Whether the session is only available via the local network and not via internet connection. Only available in some platforms. False by default */
 		bool bIsLANSession = false;
 
-		/** Whether the session is configured to run from a dedicated server. Only available in some platforms. False by default */
-		bool bIsDedicatedServerSession = false;
-
 		/** Whether this session will allow sanctioned players to join it. True by default */
 		bool bAllowSanctionedPlayers = true;
 
@@ -994,7 +991,6 @@ BEGIN_ONLINE_STRUCT_META(FCreateSession::Params)
 	ONLINE_STRUCT_FIELD(FCreateSession::Params, SessionIdOverride),
 	ONLINE_STRUCT_FIELD(FCreateSession::Params, bPresenceEnabled),
 	ONLINE_STRUCT_FIELD(FCreateSession::Params, bIsLANSession),
-	ONLINE_STRUCT_FIELD(FCreateSession::Params, bIsDedicatedServerSession),
 	ONLINE_STRUCT_FIELD(FCreateSession::Params, bAllowSanctionedPlayers),
 	ONLINE_STRUCT_FIELD(FCreateSession::Params, bAntiCheatProtected),
 	ONLINE_STRUCT_FIELD(FCreateSession::Params, SessionSettings)
diff --git a/Engine/Plugins/Online/Experimental/OnlineServicesEOSGS/Source/Private/Online/SessionsEOSGS.cpp b/Engine/Plugins/Online/Experimental/OnlineServicesEOSGS/Source/Private/Online/SessionsEOSGS.cpp
--- a/Engine/Plugins/Online/Experimental/OnlineServicesEOSGS/Source/Private/Online/SessionsEOSGS.cpp
+++ b/Engine/Plugins/Online/Experimental/OnlineServicesEOSGS/Source/Private/Online/SessionsEOSGS.cpp
@@ -551,7 +551,7 @@ void FSessionsEOSGS::WriteCreateSessionModificationHandle(EOS_HSessionModificati
 	// We won't copy bIsLANSession since it' irrelevant for EOS Sessions
 	AddAttribute(SessionModificationHandle, EOSGS_ALLOW_NEW_MEMBERS_ATTRIBUTE_KEY, { FSchemaVariant(Params.SessionSettings.bAllowNewMembers), ESchemaAttributeVisibility::Public });
 	AddAttribute(SessionModificationHandle, EOSGS_ANTI_CHEAT_PROTECTED_ATTRIBUTE_KEY, { FSchemaVariant(Params.bAntiCheatProtected), ESchemaAttributeVisibility::Public });
-	AddAttribute(SessionModificationHandle, EOSGS_IS_DEDICATED_SERVER_SESSION_ATTRIBUTE_KEY, { FSchemaVariant(Params.bIsDedicatedServerSession), ESchemaAttributeVisibility::Public });
+	AddAttribute(SessionModificationHandle, EOSGS_IS_DEDICATED_SERVER_SESSION_ATTRIBUTE_KEY, { FSchemaVariant(IsRunningDedicatedServer()), ESchemaAttributeVisibility::Public });
 
 	SetPermissionLevel(SessionModificationHandle, Params.SessionSettings.JoinPolicy);
 
diff --git a/Engine/Plugins/Online/Experimental/OnlineServicesOSSAdapter/Source/Private/Online/SessionsOSSAdapter.cpp b/Engine/Plugins/Online/Experimental/OnlineServicesOSSAdapter/Source/Private/Online/SessionsOSSAdapter.cpp
--- a/Engine/Plugins/Online/Experimental/OnlineServicesOSSAdapter/Source/Private/Online/SessionsOSSAdapter.cpp
+++ b/Engine/Plugins/Online/Experimental/OnlineServicesOSSAdapter/Source/Private/Online/SessionsOSSAdapter.cpp
@@ -931,7 +931,7 @@ FOnlineSessionSettings FSessionsOSSAdapter::BuildV1SettingsForCreate(const FCrea
 	Result.bAllowJoinViaPresenceFriendsOnly = Params.SessionSettings.bAllowNewMembers;
 	Result.Settings.Add(OSS_ADAPTER_SESSIONS_ALLOW_SANCTIONED_PLAYERS, Params.bAllowSanctionedPlayers);
 	Result.bAntiCheatProtected = Params.bAntiCheatProtected;
-	Result.bIsDedicated = Params.bIsDedicatedServerSession;
+	Result.bIsDedicated = IsRunningDedicatedServer();
 	Result.bIsLANMatch = Params.bIsLANSession;
 	Result.bShouldAdvertise = Params.SessionSettings.bAllowNewMembers;
 	if (const FCustomSessionSetting* BuildUniqueId = Params.SessionSettings.CustomSettings.Find(OSS_ADAPTER_SESSIONS_BUILD_UNIQUE_ID))
