Refactor of FSessionInvite into ISessionInvite and FSessionInviteCommon to enable logging.

diff --git a/Engine/Plugins/Online/Experimental/OnlineServices/Source/OnlineServicesCommon/Private/Online/SessionsCommon.cpp b/Engine/Plugins/Online/Experimental/OnlineServices/Source/OnlineServicesCommon/Private/Online/SessionsCommon.cpp
--- a/Engine/Plugins/Online/Experimental/OnlineServices/Source/OnlineServicesCommon/Private/Online/SessionsCommon.cpp
+++ b/Engine/Plugins/Online/Experimental/OnlineServices/Source/OnlineServicesCommon/Private/Online/SessionsCommon.cpp
@@ -85,6 +85,23 @@ namespace UE::Online {
 		}
 	}
 
+	/** FSessionInvitecommon */
+
+	FSessionInviteCommon::FSessionInviteCommon(const FAccountId& InRecipientId, const FAccountId& InSenderId, const FSessionInviteId& InInviteId, const FOnlineSessionId& InSessionId)
+		: RecipientId(InRecipientId)
+		, SenderId(InSenderId)
+		, InviteId(InInviteId)
+		, SessionId(InSessionId)
+	{
+
+	}
+
+	FString FSessionInviteCommon::ToLogString() const
+	{
+		return FString::Printf(TEXT("ISessionInvite: RecipientId [%s], SenderId [%s], InviteId [%s], SessionId [%s]"),
+			*UE::Online::ToLogString(RecipientId),*UE::Online::ToLogString(SenderId), *UE::Online::ToLogString(InviteId), *UE::Online::ToLogString(SessionId));
+	}
+
 	/** FSessionsCommon */
 
 	FSessionsCommon::FSessionsCommon(FOnlineServicesCommon& InServices)
@@ -272,9 +289,9 @@ namespace UE::Online {
 		CHECK_PARAMS_ID_HANDLE(Params.LocalAccountId, FGetSessionInviteById)
 		CHECK_PARAMS_ID_HANDLE(Params.SessionInviteId, FGetSessionInviteById)
 
-		if (const TMap<FSessionInviteId, TSharedRef<FSessionInvite>>* UserMap = SessionInvitesUserMap.Find(Params.LocalAccountId))
+		if (const TMap<FSessionInviteId, TSharedRef<FSessionInviteCommon>>* UserMap = SessionInvitesUserMap.Find(Params.LocalAccountId))
 		{
-			if (const TSharedRef<FSessionInvite>* SessionInvite = UserMap->Find(Params.SessionInviteId))
+			if (const TSharedRef<FSessionInviteCommon>* SessionInvite = UserMap->Find(Params.SessionInviteId))
 			{
 				return TOnlineResult<FGetSessionInviteById>({ *SessionInvite });
 			}
@@ -295,11 +312,11 @@ namespace UE::Online {
 	{
 		CHECK_PARAMS_ID_HANDLE(Params.LocalAccountId, FGetAllSessionInvites)
 		
-		TArray<TSharedRef<const FSessionInvite>> SessionInvites;
+		TArray<TSharedRef<const ISessionInvite>> SessionInvites;
 
-		if (const TMap<FSessionInviteId, TSharedRef<FSessionInvite>>* UserMap = SessionInvitesUserMap.Find(Params.LocalAccountId))
+		if (const TMap<FSessionInviteId, TSharedRef<FSessionInviteCommon>>* UserMap = SessionInvitesUserMap.Find(Params.LocalAccountId))
 		{
-			for (const TPair<FSessionInviteId, TSharedRef<FSessionInvite>>& Entry : *UserMap)
+			for (const TPair<FSessionInviteId, TSharedRef<FSessionInviteCommon>>& Entry : *UserMap)
 			{
 				SessionInvites.Add(Entry.Value);
 			}
@@ -833,7 +850,7 @@ namespace UE::Online {
 
 	TOptional<FOnlineError> FSessionsCommon::CheckState(const FRejectSessionInvite::Params& Params) const
 	{
-		if (const TMap<FSessionInviteId, TSharedRef<FSessionInvite>>* UserMap = SessionInvitesUserMap.Find(Params.LocalAccountId))
+		if (const TMap<FSessionInviteId, TSharedRef<FSessionInviteCommon>>* UserMap = SessionInvitesUserMap.Find(Params.LocalAccountId))
 		{
 			if (!UserMap->Contains(Params.SessionInviteId))
 			{
@@ -883,11 +900,11 @@ namespace UE::Online {
 
 	// Auxiliary methods
 
-	void FSessionsCommon::AddSessionInvite(const TSharedRef<FSessionInvite> SessionInvite, const TSharedRef<FSessionCommon> Session, const FAccountId& LocalAccountId)
+	void FSessionsCommon::AddSessionInvite(const TSharedRef<FSessionInviteCommon> SessionInvite, const TSharedRef<FSessionCommon> Session, const FAccountId& LocalAccountId)
 	{
 		AllSessionsById.Emplace(Session->GetSessionId(), Session);
 
-		TMap<FSessionInviteId, TSharedRef<FSessionInvite>>& SessionInvitesMap = SessionInvitesUserMap.FindOrAdd(LocalAccountId);
+		TMap<FSessionInviteId, TSharedRef<FSessionInviteCommon>>& SessionInvitesMap = SessionInvitesUserMap.FindOrAdd(LocalAccountId);
 		SessionInvitesMap.Emplace(SessionInvite->InviteId, SessionInvite);
 	}
 
@@ -920,10 +937,10 @@ namespace UE::Online {
 
 	void FSessionsCommon::ClearSessionInvitesForSession(const FAccountId& LocalAccountId, const FOnlineSessionId SessionId)
 	{
-		if (TMap<FSessionInviteId, TSharedRef<FSessionInvite>>* UserMap = SessionInvitesUserMap.Find(LocalAccountId))
+		if (TMap<FSessionInviteId, TSharedRef<FSessionInviteCommon>>* UserMap = SessionInvitesUserMap.Find(LocalAccountId))
 		{
 			TArray<FSessionInviteId> InviteIdsToRemove;
-			for (const TPair<FSessionInviteId, TSharedRef<FSessionInvite>>& Entry : *UserMap)
+			for (const TPair<FSessionInviteId, TSharedRef<FSessionInviteCommon>>& Entry : *UserMap)
 			{
 				if (Entry.Value->SessionId == SessionId)
 				{
@@ -980,9 +997,9 @@ namespace UE::Online {
 			}
 		}
 
-		for (const TPair<FAccountId, TMap<FSessionInviteId, TSharedRef<FSessionInvite>>>& Entry : SessionInvitesUserMap)
+		for (const TPair<FAccountId, TMap<FSessionInviteId, TSharedRef<FSessionInviteCommon>>>& Entry : SessionInvitesUserMap)
 		{
-			for (const TPair<FSessionInviteId, TSharedRef<FSessionInvite>>& InviteMap : Entry.Value)
+			for (const TPair<FSessionInviteId, TSharedRef<FSessionInviteCommon>>& InviteMap : Entry.Value)
 			{
 				if (InviteMap.Value->SessionId == SessionId)
 				{
diff --git a/Engine/Plugins/Online/Experimental/OnlineServices/Source/OnlineServicesCommon/Public/Online/SessionsCommon.h b/Engine/Plugins/Online/Experimental/OnlineServices/Source/OnlineServicesCommon/Public/Online/SessionsCommon.h
--- a/Engine/Plugins/Online/Experimental/OnlineServices/Source/OnlineServicesCommon/Public/Online/SessionsCommon.h
+++ b/Engine/Plugins/Online/Experimental/OnlineServices/Source/OnlineServicesCommon/Public/Online/SessionsCommon.h
@@ -171,6 +171,34 @@ public:
 	FSessionCommon& operator+=(const FSessionUpdate& SessionUpdate);
 };
 
+class ONLINESERVICESCOMMON_API FSessionInviteCommon : public ISessionInvite
+{
+public:
+	FSessionInviteCommon() = default;
+	FSessionInviteCommon(const FSessionInviteCommon& Other) = default;
+	FSessionInviteCommon(const FAccountId& InRecipientId, const FAccountId& InSenderId, const FSessionInviteId& InInviteId, const FOnlineSessionId& InSessionId);
+
+	virtual const FAccountId GetRecipientId() const override { return RecipientId; }
+	virtual const FAccountId GetSenderId() const override { return SenderId; }
+	virtual const FSessionInviteId GetInviteId() const override { return InviteId; }
+	virtual const FOnlineSessionId GetSessionId() const override { return SessionId; }
+
+	virtual FString ToLogString() const override;
+
+public:
+	/* The id handle for the user which the invite got sent to */
+	FAccountId RecipientId;
+
+	/* The id handle for the user which sent the invite */
+	FAccountId SenderId;
+
+	/* The invite id handle, needed for retrieving invite information and rejecting the invite */
+	FSessionInviteId InviteId;
+
+	/* The session id handle, needed for retrieving the session information */
+	FOnlineSessionId SessionId;
+};
+
 struct FGetMutableSessionByName
 {
 	static constexpr TCHAR Name[] = TEXT("GetMutableSessionByName");
@@ -248,7 +276,7 @@ protected:
 	TOnlineResult<FGetMutableSessionByName> GetMutableSessionByName(FGetMutableSessionByName::Params&& Params) const;
 	TOnlineResult<FGetMutableSessionById> GetMutableSessionById(FGetMutableSessionById::Params&& Params) const;
 
-	void AddSessionInvite(const TSharedRef<FSessionInvite> SessionInvite, const TSharedRef<FSessionCommon> Session, const FAccountId& LocalAccountId);
+	void AddSessionInvite(const TSharedRef<FSessionInviteCommon> SessionInvite, const TSharedRef<FSessionCommon> Session, const FAccountId& LocalAccountId);
 	void AddSearchResult(const TSharedRef<FSessionCommon> Session, const FAccountId& LocalAccountId);
 	void AddSessionWithReferences(const TSharedRef<FSessionCommon> Session, const FName& SessionName, const FAccountId& LocalAccountId, bool bIsPresenceSession);
 	void AddSessionReferences(const FOnlineSessionId SessionId, const FName& SessionName, const FAccountId& LocalAccountId, bool bIsPresenceSession);
@@ -340,7 +368,7 @@ protected:
 	TMap<FAccountId, FOnlineSessionId> PresenceSessionsUserMap;
 
 	/** Cache for received session invites, mapped per user */
-	TMap<FAccountId, TMap<FSessionInviteId, TSharedRef<FSessionInvite>>> SessionInvitesUserMap;
+	TMap<FAccountId, TMap<FSessionInviteId, TSharedRef<FSessionInviteCommon>>> SessionInvitesUserMap;
 
 	/** Cache for the last set of session search results, mapped per user */
 	TMap<FAccountId, TArray<FOnlineSessionId>> SearchResultsUserMap;
diff --git a/Engine/Plugins/Online/Experimental/OnlineServices/Source/OnlineServicesInterface/Private/Online/Sessions.cpp b/Engine/Plugins/Online/Experimental/OnlineServices/Source/OnlineServicesInterface/Private/Online/Sessions.cpp
--- a/Engine/Plugins/Online/Experimental/OnlineServices/Source/OnlineServicesInterface/Private/Online/Sessions.cpp
+++ b/Engine/Plugins/Online/Experimental/OnlineServices/Source/OnlineServicesInterface/Private/Online/Sessions.cpp
@@ -77,7 +77,8 @@ FSessionSettingsUpdate& FSessionSettingsUpdate::operator+=(FSessionSettingsUpdat
 
 FSessionSettings& FSessionSettings::operator+=(const FSessionSettingsChanges& SettingsChanges)
 {
-	COPY_TOPTIONAL_VALUE_IF_SET(SchemaName) // TODO: We may need some additional logic for schema changes
+	// TODO: Full support for Schema functionality still pending
+	COPY_TOPTIONAL_VALUE_IF_SET(SchemaName)
 	COPY_TOPTIONAL_VALUE_IF_SET(NumMaxConnections)
 	COPY_TOPTIONAL_VALUE_IF_SET(JoinPolicy)
 	COPY_TOPTIONAL_VALUE_IF_SET(bAllowNewMembers)
@@ -107,6 +108,11 @@ const FString ToLogString(const ISession& Session)
 	return Session.ToLogString();
 }
 
+const FString ToLogString(const ISessionInvite& SessionInvite)
+{
+	return SessionInvite.ToLogString();
+}
+
 const TCHAR* LexToString(EUISessionJoinRequestedSource UISessionJoinRequestedSource)
 {
 	switch (UISessionJoinRequestedSource)
diff --git a/Engine/Plugins/Online/Experimental/OnlineServices/Source/OnlineServicesInterface/Public/Online/Sessions.h b/Engine/Plugins/Online/Experimental/OnlineServices/Source/OnlineServicesInterface/Public/Online/Sessions.h
--- a/Engine/Plugins/Online/Experimental/OnlineServices/Source/OnlineServicesInterface/Public/Online/Sessions.h
+++ b/Engine/Plugins/Online/Experimental/OnlineServices/Source/OnlineServicesInterface/Public/Online/Sessions.h
@@ -187,23 +187,21 @@ public:
 };
 ONLINESERVICESINTERFACE_API const FString ToLogString(const ISession& Session);
 
-/** Information about an invitation to join an Online Session */
-struct FSessionInvite
+/** Interface to access all information about an invitation to join an Online Session. Read only */
+class ISessionInvite
 {
-	/* The id handle for the user which the invite got sent to */
-	FAccountId RecipientId;
-
-	/* The id handle for the user which sent the invite */
-	FAccountId SenderId;
-
-	/* The invite id handle, needed for retrieving invite information and rejecting the invite */
-	FSessionInviteId InviteId;
+public:
+	/** Retrieves the id handle for the user who received the invite */
+	virtual const FAccountId GetRecipientId() const = 0;
 
-	/* The session id handle, needed for retrieving the session information */
-	FOnlineSessionId SessionId;
+	/** Retrieves the id handle for the user who sent the invite */
+	virtual const FAccountId GetSenderId() const = 0;
+	virtual const FSessionInviteId GetInviteId() const = 0;
+	virtual const FOnlineSessionId GetSessionId() const = 0;
 
-	// TODO: Refactor into ISessionInvite, FSessionInviteCommon/EOSGS/OSSAdapter to add a ToLogString method
+	virtual FString ToLogString() const = 0;
 };
+ONLINESERVICESINTERFACE_API const FString ToLogString(const ISessionInvite& Session);
 
 // ISessions Methods
 
@@ -564,7 +562,7 @@ struct FGetSessionInviteById
 	struct Result
 	{
 		/** Reference to the invite mapped to the given id */
-		TSharedRef<const FSessionInvite> SessionInvite;
+		TSharedRef<const ISessionInvite> SessionInvite;
 	};
 };
 
@@ -581,7 +579,7 @@ struct FGetAllSessionInvites
 	struct Result
 	{
 		/** Array of invites received by the given user */
-		TArray<TSharedRef<const FSessionInvite>> SessionInvites;
+		TArray<TSharedRef<const ISessionInvite>> SessionInvites;
 	};
 };
 
@@ -922,13 +920,6 @@ BEGIN_ONLINE_STRUCT_META(FSessionSettingsUpdate)
 	ONLINE_STRUCT_FIELD(FSessionSettingsUpdate, RemovedCustomSettings)
 END_ONLINE_STRUCT_META()
 
-BEGIN_ONLINE_STRUCT_META(FSessionInvite)
-	ONLINE_STRUCT_FIELD(FSessionInvite, RecipientId),
-	ONLINE_STRUCT_FIELD(FSessionInvite, SenderId),
-	ONLINE_STRUCT_FIELD(FSessionInvite, InviteId),
-	ONLINE_STRUCT_FIELD(FSessionInvite, SessionId)
-END_ONLINE_STRUCT_META()
-
 BEGIN_ONLINE_STRUCT_META(FGetAllSessions::Params)
 	ONLINE_STRUCT_FIELD(FGetAllSessions::Params, LocalAccountId)
 END_ONLINE_STRUCT_META()
diff --git a/Engine/Plugins/Online/Experimental/OnlineServicesEOSGS/Source/Private/Online/SessionsEOSGS.cpp b/Engine/Plugins/Online/Experimental/OnlineServicesEOSGS/Source/Private/Online/SessionsEOSGS.cpp
--- a/Engine/Plugins/Online/Experimental/OnlineServicesEOSGS/Source/Private/Online/SessionsEOSGS.cpp
+++ b/Engine/Plugins/Online/Experimental/OnlineServicesEOSGS/Source/Private/Online/SessionsEOSGS.cpp
@@ -232,12 +232,12 @@ void FSessionsEOSGS::HandleSessionInviteReceived(const EOS_Sessions_SessionInvit
 
 					TSharedRef<FSessionCommon> Session = Result.GetOkValue().Session;
 
-					TSharedRef<FSessionInvite> SessionInviteRef = MakeShared<FSessionInvite>(FSessionInvite{
+					TSharedRef<FSessionInviteCommon> SessionInviteRef = MakeShared<FSessionInviteCommon>(
 						ReceiverId,
 						SenderId,
 						SessionInviteId,
 						Session->GetSessionId()
-						});
+					);
 
 					AddSessionInvite(SessionInviteRef, Session, ReceiverId);
 
@@ -1464,7 +1464,7 @@ TFuture<TOnlineResult<FRejectSessionInvite>> FSessionsEOSGS::RejectSessionInvite
 				return;
 			}
 
-			if (TMap<FSessionInviteId, TSharedRef<FSessionInvite>>* UserMap = SessionInvitesUserMap.Find(Params.LocalAccountId))
+			if (TMap<FSessionInviteId, TSharedRef<FSessionInviteCommon>>* UserMap = SessionInvitesUserMap.Find(Params.LocalAccountId))
 			{
 				UserMap->Remove(Params.SessionInviteId);
 			}
diff --git a/Engine/Plugins/Online/Experimental/OnlineServicesOSSAdapter/Source/Private/Online/SessionsOSSAdapter.cpp b/Engine/Plugins/Online/Experimental/OnlineServicesOSSAdapter/Source/Private/Online/SessionsOSSAdapter.cpp
--- a/Engine/Plugins/Online/Experimental/OnlineServicesOSSAdapter/Source/Private/Online/SessionsOSSAdapter.cpp
+++ b/Engine/Plugins/Online/Experimental/OnlineServicesOSSAdapter/Source/Private/Online/SessionsOSSAdapter.cpp
@@ -294,24 +294,7 @@ void FSessionsOSSAdapter::Initialize()
 	IdentityInterface = SubsystemV1.GetIdentityInterface();
 	check(IdentityInterface);
 
-	SessionsInterface->AddOnSessionInviteReceivedDelegate_Handle(FOnSessionInviteReceivedDelegate::CreateLambda([this](const FUniqueNetId& UserId, const FUniqueNetId& FromId, const FString& AppId, const FOnlineSessionSearchResult& InviteResult)
-	{
-		FOnlineServicesOSSAdapter& ServicesOSSAdapter = static_cast<FOnlineServicesOSSAdapter&>(Services);
-
-		TSharedRef<FSessionCommon> Session = BuildV2Session(&InviteResult.Session);
-
-		TSharedRef<FSessionInvite> SessionInvite = MakeShared<FSessionInvite>();
-		SessionInvite->SenderId = ServicesOSSAdapter.GetAccountIdRegistry().FindOrAddHandle(FromId.AsShared());
-		SessionInvite->RecipientId = ServicesOSSAdapter.GetAccountIdRegistry().FindOrAddHandle(UserId.AsShared());
-		SessionInvite->InviteId = GetSessionInviteIdRegistry().BasicRegistry.FindOrAddHandle(AppId);
-		SessionInvite->SessionId = Session->GetSessionId();
-
-		AddSessionInvite(SessionInvite, Session, SessionInvite->RecipientId);
-
-		FSessionInviteReceived Event{ ServicesOSSAdapter.GetAccountIdRegistry().FindOrAddHandle(UserId.AsShared()), SessionInvite->InviteId };
-
-		SessionEvents.OnSessionInviteReceived.Broadcast(Event);
-	}));
+	// OnSessionInviteReceived is not used by current implementations of Sessions in V1, so we won't bind to FOnSessionInviteReceived
 
 	SessionsInterface->AddOnSessionUserInviteAcceptedDelegate_Handle(FOnSessionUserInviteAcceptedDelegate::CreateLambda([this](const bool bWasSuccessful, const int32 ControllerId, FUniqueNetIdPtr InvitedUserIdPtr, const FOnlineSessionSearchResult& InviteResult)
 	{
