Fix for typo in return type

diff --git a/Engine/Plugins/Online/Experimental/OnlineServices/Source/OnlineServicesCommon/Public/Online/SessionsCommon.h b/Engine/Plugins/Online/Experimental/OnlineServices/Source/OnlineServicesCommon/Public/Online/SessionsCommon.h
--- a/Engine/Plugins/Online/Experimental/OnlineServices/Source/OnlineServicesCommon/Public/Online/SessionsCommon.h
+++ b/Engine/Plugins/Online/Experimental/OnlineServices/Source/OnlineServicesCommon/Public/Online/SessionsCommon.h
@@ -139,7 +139,7 @@ public:
 	virtual const FOnlineSessionId GetSessionId() const override			{ return GetSessionInfo().SessionId; }
 	virtual const uint32 GetNumOpenConnections() const override				{ return SessionSettings.NumMaxConnections - SessionMembers.Num(); }
 	virtual const FSessionInfo& GetSessionInfo() const override				{ return SessionInfo; }
-	virtual const FSessionSettings GetSessionSettings() const override		{ return SessionSettings; }
+	virtual const FSessionSettings& GetSessionSettings() const override		{ return SessionSettings; }
 	virtual const FSessionMemberIdsSet& GetSessionMembers() const override	{ return SessionMembers; }
 
 	virtual bool IsJoinable() const override								{ return GetNumOpenConnections() > 0 && SessionSettings.bAllowNewMembers; }
diff --git a/Engine/Plugins/Online/Experimental/OnlineServices/Source/OnlineServicesInterface/Public/Online/Sessions.h b/Engine/Plugins/Online/Experimental/OnlineServices/Source/OnlineServicesInterface/Public/Online/Sessions.h
--- a/Engine/Plugins/Online/Experimental/OnlineServices/Source/OnlineServicesInterface/Public/Online/Sessions.h
+++ b/Engine/Plugins/Online/Experimental/OnlineServices/Source/OnlineServicesInterface/Public/Online/Sessions.h
@@ -171,7 +171,7 @@ public:
 	virtual const FSessionInfo& GetSessionInfo() const = 0;
 
 	/** Retrieves the set of variable information about the session */
-	virtual const FSessionSettings GetSessionSettings() const = 0;
+	virtual const FSessionSettings& GetSessionSettings() const = 0;
 
 	/** Retrieves the list of users currently in the session */
 	virtual const FSessionMemberIdsSet& GetSessionMembers() const = 0;
