Fix for inconsistency in LAN Sessions behaviour (adding member on creation)

diff --git a/Engine/Plugins/Online/Experimental/OnlineServices/Source/OnlineServicesCommonEngineUtils/Private/Online/SessionsLAN.cpp b/Engine/Plugins/Online/Experimental/OnlineServices/Source/OnlineServicesCommonEngineUtils/Private/Online/SessionsLAN.cpp
--- a/Engine/Plugins/Online/Experimental/OnlineServices/Source/OnlineServicesCommonEngineUtils/Private/Online/SessionsLAN.cpp
+++ b/Engine/Plugins/Online/Experimental/OnlineServices/Source/OnlineServicesCommonEngineUtils/Private/Online/SessionsLAN.cpp
@@ -149,9 +149,6 @@ TFuture<TOnlineResult<FCreateSession>> FSessionsLAN::CreateSessionImpl(const FCr
 	NewSessionLANRef->SessionInfo.SessionIdOverride = Params.SessionIdOverride;
 	NewSessionLANRef->SessionSettings = Params.SessionSettings;
 
-	// For LAN sessions, we'll add the Session member manually instead of calling JoinSession since there is no API calls involved
-	NewSessionLANRef->SessionMembers.Emplace(Params.LocalAccountId);
-
 	// We save the local object for the session, and set up the appropriate references
 	AddSessionWithReferences(NewSessionLANRef, Params.SessionName, Params.LocalAccountId, Params.bPresenceEnabled);
 
