OSSv2: Add the missing metadata of FQueryStats and FBatchQueryStats.

diff --git a/Engine/Plugins/Online/Experimental/OnlineServices/Source/OnlineServicesInterface/Public/Online/Stats.h b/Engine/Plugins/Online/Experimental/OnlineServices/Source/OnlineServicesInterface/Public/Online/Stats.h
--- a/Engine/Plugins/Online/Experimental/OnlineServices/Source/OnlineServicesInterface/Public/Online/Stats.h
+++ b/Engine/Plugins/Online/Experimental/OnlineServices/Source/OnlineServicesInterface/Public/Online/Stats.h
@@ -177,7 +177,8 @@ END_ONLINE_STRUCT_META()
 
 BEGIN_ONLINE_STRUCT_META(FQueryStats::Params)
 	ONLINE_STRUCT_FIELD(FQueryStats::Params, LocalAccountId),
-	ONLINE_STRUCT_FIELD(FQueryStats::Params, TargetAccountId)
+	ONLINE_STRUCT_FIELD(FQueryStats::Params, TargetAccountId),
+	ONLINE_STRUCT_FIELD(FQueryStats::Params, StatNames)
 END_ONLINE_STRUCT_META()
 
 BEGIN_ONLINE_STRUCT_META(FQueryStats::Result)
@@ -186,7 +187,8 @@ END_ONLINE_STRUCT_META()
 
 BEGIN_ONLINE_STRUCT_META(FBatchQueryStats::Params)
 	ONLINE_STRUCT_FIELD(FBatchQueryStats::Params, LocalAccountId),
-	ONLINE_STRUCT_FIELD(FBatchQueryStats::Params, TargetAccountIds)
+	ONLINE_STRUCT_FIELD(FBatchQueryStats::Params, TargetAccountIds),
+	ONLINE_STRUCT_FIELD(FBatchQueryStats::Params, StatNames)
 END_ONLINE_STRUCT_META()
 
 BEGIN_ONLINE_STRUCT_META(FBatchQueryStats::Result)
