Fix for crash discovered on application shutdown on OnlineTests
Added logic to destroy Services instances when the factory is unregistered on module shutdown

diff --git a/Engine/Plugins/Online/Experimental/OnlineServices/Source/OnlineServicesInterface/Private/Online/OnlineServices.cpp b/Engine/Plugins/Online/Experimental/OnlineServices/Source/OnlineServicesInterface/Private/Online/OnlineServices.cpp
--- a/Engine/Plugins/Online/Experimental/OnlineServices/Source/OnlineServicesInterface/Private/Online/OnlineServices.cpp
+++ b/Engine/Plugins/Online/Experimental/OnlineServices/Source/OnlineServicesInterface/Private/Online/OnlineServices.cpp
@@ -67,9 +67,14 @@ TSharedPtr<IOnlineServices> GetServices(EOnlineServices OnlineServices, FName In
 	return FOnlineServicesRegistry::Get().GetNamedServicesInstance(OnlineServices, InstanceName);
 }
 
-void DestroyServices(EOnlineServices OnlineServices, FName InstanceName)
+void DestroyService(EOnlineServices OnlineServices, FName InstanceName)
 {
 	FOnlineServicesRegistry::Get().DestroyNamedServicesInstance(OnlineServices, InstanceName);
 }
 
+void DestroyAllNamedServices(EOnlineServices OnlineServices)
+{
+	FOnlineServicesRegistry::Get().DestroyAllNamedServicesInstances(OnlineServices);
+}
+
 /* UE::Online */ }
diff --git a/Engine/Plugins/Online/Experimental/OnlineServices/Source/OnlineServicesInterface/Private/Online/OnlineServicesRegistry.cpp b/Engine/Plugins/Online/Experimental/OnlineServices/Source/OnlineServicesInterface/Private/Online/OnlineServicesRegistry.cpp
--- a/Engine/Plugins/Online/Experimental/OnlineServices/Source/OnlineServicesInterface/Private/Online/OnlineServicesRegistry.cpp
+++ b/Engine/Plugins/Online/Experimental/OnlineServices/Source/OnlineServicesInterface/Private/Online/OnlineServicesRegistry.cpp
@@ -34,6 +34,8 @@ void FOnlineServicesRegistry::UnregisterServicesFactory(EOnlineServices OnlineSe
 	{
 		ServicesFactories.Remove(OnlineServices);
 	}
+
+	DestroyAllNamedServicesInstances(OnlineServices);
 }
 
 bool FOnlineServicesRegistry::IsLoaded(EOnlineServices OnlineServices, FName InstanceName) const
@@ -113,6 +115,19 @@ void FOnlineServicesRegistry::DestroyNamedServicesInstance(EOnlineServices Onlin
 	}
 }
 
+void FOnlineServicesRegistry::DestroyAllNamedServicesInstances(EOnlineServices OnlineServices)
+{
+	if (TMap<FName, TSharedRef<IOnlineServices>>* ServicesMapPtr = NamedServiceInstances.Find(OnlineServices))
+	{
+		for (const TPair<FName, TSharedRef<IOnlineServices>>& ServicesEntryRef : *ServicesMapPtr)
+		{
+			ServicesEntryRef.Value->Destroy();
+		}
+
+		NamedServiceInstances.Remove(OnlineServices);
+	}
+}
+
 TSharedPtr<IOnlineServices> FOnlineServicesRegistry::CreateServices(EOnlineServices OnlineServices, FName InstanceName)
 {
 	TSharedPtr<IOnlineServices> Services;
diff --git a/Engine/Plugins/Online/Experimental/OnlineServices/Source/OnlineServicesInterface/Public/Online/OnlineServices.h b/Engine/Plugins/Online/Experimental/OnlineServices/Source/OnlineServicesInterface/Public/Online/OnlineServices.h
--- a/Engine/Plugins/Online/Experimental/OnlineServices/Source/OnlineServicesInterface/Public/Online/OnlineServices.h
+++ b/Engine/Plugins/Online/Experimental/OnlineServices/Source/OnlineServicesInterface/Public/Online/OnlineServices.h
@@ -203,7 +203,14 @@ TSharedPtr<ServicesClass> GetServices(FName InstanceName = NAME_None)
  * @param OnlineServices Type of online services to destroy
  * @param InstanceName Name of the services instance to destroy
  */
-ONLINESERVICESINTERFACE_API void DestroyServices(EOnlineServices OnlineServices = EOnlineServices::Default, FName InstanceName = NAME_None);
+ONLINESERVICESINTERFACE_API void DestroyService(EOnlineServices OnlineServices = EOnlineServices::Default, FName InstanceName = NAME_None);
+
+/**
+ * Destroy all instances of the online service specified
+ *
+ * @param OnlineServices Type of online services to destroy all instances from
+ */
+ONLINESERVICESINTERFACE_API void DestroyAllNamedServices(EOnlineServices OnlineServices);
 
 namespace Meta {
 // TODO: Move to OnlineServices_Meta.inl file?
diff --git a/Engine/Plugins/Online/Experimental/OnlineServices/Source/OnlineServicesInterface/Public/Online/OnlineServicesRegistry.h b/Engine/Plugins/Online/Experimental/OnlineServices/Source/OnlineServicesInterface/Public/Online/OnlineServicesRegistry.h
--- a/Engine/Plugins/Online/Experimental/OnlineServices/Source/OnlineServicesInterface/Public/Online/OnlineServicesRegistry.h
+++ b/Engine/Plugins/Online/Experimental/OnlineServices/Source/OnlineServicesInterface/Public/Online/OnlineServicesRegistry.h
@@ -90,6 +90,13 @@ public:
 	 */
 	ONLINESERVICESINTERFACE_API void DestroyNamedServicesInstance(EOnlineServices OnlineServices, FName InstanceName);
 
+	/**
+	 * Destroy all instances of a specific OnlineServices
+	 *
+	 * @param OnlineServices  Type of online services for the IOnlineServices instance
+	 */
+	ONLINESERVICESINTERFACE_API void DestroyAllNamedServicesInstances(EOnlineServices OnlineServices);
+
 	/**
 	 * Create and initialize a new IOnlineServices instance
 	 *
diff --git a/Engine/Plugins/Online/Experimental/OnlineServicesEOS/Source/Private/OnlineServicesEOSModule.cpp b/Engine/Plugins/Online/Experimental/OnlineServicesEOS/Source/Private/OnlineServicesEOSModule.cpp
--- a/Engine/Plugins/Online/Experimental/OnlineServicesEOS/Source/Private/OnlineServicesEOSModule.cpp
+++ b/Engine/Plugins/Online/Experimental/OnlineServicesEOS/Source/Private/OnlineServicesEOSModule.cpp
@@ -15,6 +15,7 @@ class FOnlineServicesEOSModule : public IModuleInterface
 {
 public:
 	virtual void StartupModule() override;
+	virtual void ShutdownModule() override;
 protected:
 };
 
@@ -43,6 +44,14 @@ void FOnlineServicesEOSModule::StartupModule()
 	FOnlineIdRegistryRegistry::Get().RegisterAccountIdRegistry(EOnlineServices::Epic, &FOnlineAccountIdRegistryEOS::Get(), Priority);
 }
 
+void FOnlineServicesEOSModule::ShutdownModule()
+{
+	// Make this higher priority that EOSGS
+	const int Priority = 1;
+	FOnlineServicesRegistry::Get().UnregisterServicesFactory(EOnlineServices::Epic, Priority);
+	FOnlineIdRegistryRegistry::Get().UnregisterAccountIdRegistry(EOnlineServices::Epic, Priority);
+}
+
 /* UE::Online */ }
 
 IMPLEMENT_MODULE(UE::Online::FOnlineServicesEOSModule, OnlineServicesEOS);
diff --git a/Engine/Plugins/Online/Experimental/OnlineServicesEOSGS/Source/Private/OnlineServicesEOSGSModule.cpp b/Engine/Plugins/Online/Experimental/OnlineServicesEOSGS/Source/Private/OnlineServicesEOSGSModule.cpp
--- a/Engine/Plugins/Online/Experimental/OnlineServicesEOSGS/Source/Private/OnlineServicesEOSGSModule.cpp
+++ b/Engine/Plugins/Online/Experimental/OnlineServicesEOSGS/Source/Private/OnlineServicesEOSGSModule.cpp
@@ -45,6 +45,11 @@ void FOnlineServicesEOSGSModule::StartupModule()
 
 void FOnlineServicesEOSGSModule::ShutdownModule()
 {
+	FOnlineServicesRegistry::Get().UnregisterServicesFactory(EOnlineServices::Epic);
+	FOnlineIdRegistryRegistry::Get().UnregisterAccountIdRegistry(EOnlineServices::Epic);
+	FOnlineIdRegistryRegistry::Get().UnregisterSessionIdRegistry(EOnlineServices::Epic);
+	FOnlineIdRegistryRegistry::Get().UnregisterSessionInviteIdRegistry(EOnlineServices::Epic);
+
 	UE::Online::FOnlineServicesEOSGSPlatformFactory::TearDown();
 }
 
diff --git a/Engine/Plugins/Online/Experimental/OnlineServicesNull/Source/Private/Online/OnlineServicesNullModule.cpp b/Engine/Plugins/Online/Experimental/OnlineServicesNull/Source/Private/Online/OnlineServicesNullModule.cpp
--- a/Engine/Plugins/Online/Experimental/OnlineServicesNull/Source/Private/Online/OnlineServicesNullModule.cpp
+++ b/Engine/Plugins/Online/Experimental/OnlineServicesNull/Source/Private/Online/OnlineServicesNullModule.cpp
@@ -16,6 +16,7 @@ class FOnlineServicesNullModule : public IModuleInterface
 {
 public:
 	virtual void StartupModule() override;
+	virtual void ShutdownModule() override;
 protected:
 };
 
@@ -37,6 +38,13 @@ void FOnlineServicesNullModule::StartupModule()
 	FOnlineIdRegistryRegistry::Get().RegisterSessionIdRegistry(EOnlineServices::Null, &FOnlineSessionIdRegistryNull::Get());
 }
 
+void FOnlineServicesNullModule::ShutdownModule()
+{
+	FOnlineServicesRegistry::Get().UnregisterServicesFactory(EOnlineServices::Null);
+	FOnlineIdRegistryRegistry::Get().UnregisterAccountIdRegistry(EOnlineServices::Null);
+	FOnlineIdRegistryRegistry::Get().UnregisterSessionIdRegistry(EOnlineServices::Null);
+}
+
 /* UE::Online */ }
 
 IMPLEMENT_MODULE(UE::Online::FOnlineServicesNullModule, OnlineServicesNull);
diff --git a/Engine/Plugins/Online/OnlineSubsystemUtils/Source/OnlineSubsystemUtils/Private/OnlineServicesEngineInterfaceImpl.cpp b/Engine/Plugins/Online/OnlineSubsystemUtils/Source/OnlineSubsystemUtils/Private/OnlineServicesEngineInterfaceImpl.cpp
--- a/Engine/Plugins/Online/OnlineSubsystemUtils/Source/OnlineSubsystemUtils/Private/OnlineServicesEngineInterfaceImpl.cpp
+++ b/Engine/Plugins/Online/OnlineSubsystemUtils/Source/OnlineSubsystemUtils/Private/OnlineServicesEngineInterfaceImpl.cpp
@@ -65,7 +65,7 @@ void UOnlineServicesEngineInterfaceImpl::ShutdownOnlineSubsystem(FName OnlineIde
 void UOnlineServicesEngineInterfaceImpl::DestroyOnlineSubsystem(FName OnlineIdentifier)
 {
 	// TODO:  Does this need to support multiple online service types?  Other accessors seem to not differentiate and just use the default
-	UE::Online::DestroyServices(UE::Online::EOnlineServices::Default, OnlineIdentifier);
+	UE::Online::DestroyService(UE::Online::EOnlineServices::Default, OnlineIdentifier);
 }
 
 FName UOnlineServicesEngineInterfaceImpl::GetDefaultOnlineSubsystemName() const
