OSSv2: Output console exe result in console window.

diff --git a/Engine/Plugins/Online/OnlineServices/Source/OnlineServicesCommon/Public/Online/OnlineExecHandler.h b/Engine/Plugins/Online/OnlineServices/Source/OnlineServicesCommon/Public/Online/OnlineExecHandler.h
--- a/Engine/Plugins/Online/OnlineServices/Source/OnlineServicesCommon/Public/Online/OnlineExecHandler.h
+++ b/Engine/Plugins/Online/OnlineServices/Source/OnlineServicesCommon/Public/Online/OnlineExecHandler.h
@@ -47,7 +47,7 @@
 			- commas are optional- [5 3 7 9] is also valid
 			- arrays support internal objects/recursive arrays, e.g. [[2,3], [4,5]]
 		for TMap, you can use {} syntax. 
-			- [licenseData=licenseDataText, kNumUsers=5}
+			- {licenseData=licenseDataText, kNumUsers=5}
 			- please note: TMap does not currently support rich syntax so nested objects/arrays/etc may not parse properly
 
 	Full command examples:
@@ -898,7 +898,7 @@ public:
 			bool bResult = Private::ParseOnlineExecParams(Cmd, Field, Services);
 			if (!bResult)
 			{
-				UE_LOG(LogOnlineServices, Warning, TEXT("Failed to resolve outer field %s"), Name);
+				UE_LOG(LogConsoleResponse, Warning, TEXT("Failed to resolve outer field %s"), Name);
 			}
 			bSuccess &= bResult;
 		});
@@ -914,13 +914,13 @@ public:
 			TOnlineAsyncOpHandle<OpType> AsyncOpHandle = Invoke(Function, Interface, MoveTemp(Params));
 			AsyncOpHandle.OnComplete([&Ar](const TOnlineResult<OpType>& Result)
 				{
-					UE_LOG(LogOnlineServices, Log, TEXT("%s result: %s"), OpType::Name, *ToLogString(Result));
+					UE_LOG(LogConsoleResponse, Display, TEXT("%s result: %s"), OpType::Name, *ToLogString(Result));
 				});
 		}
 		else
 		{
 			TOnlineResult<OpType> Result = Invoke(Function, Interface, MoveTemp(Params));
-			UE_LOG(LogOnlineServices, Log, TEXT("%s result: %s"), OpType::Name, *ToLogString(Result));
+			UE_LOG(LogConsoleResponse, Display, TEXT("%s result: %s"), OpType::Name, *ToLogString(Result));
 		}
 
 		return true;
