Fix online ConstructDelegate helper compile errors with non-default constructible types

#tests Added test cases with non-default constructible type to OnlineSyntaxTests

diff --git a/Engine/Plugins/Online/OnlineServices/Source/OnlineServicesCommon/Private/Online/OnlineSyntaxTests.cpp b/Engine/Plugins/Online/OnlineServices/Source/OnlineServicesCommon/Private/Online/OnlineSyntaxTests.cpp
--- a/Engine/Plugins/Online/OnlineServices/Source/OnlineServicesCommon/Private/Online/OnlineSyntaxTests.cpp
+++ b/Engine/Plugins/Online/OnlineServices/Source/OnlineServicesCommon/Private/Online/OnlineSyntaxTests.cpp
@@ -246,6 +246,36 @@ public:
 inline void FN(int, const FString&) {}
 inline void FN2(int, const FString&, float) {}
 
+class FOnlineServicesTest : public UE::Online::FOnlineServicesCommon
+{
+public:
+	using Super = UE::Online::FOnlineServicesCommon;
+
+	FOnlineServicesTest(FName InInstanceName) : Super(TEXT("Test"), InInstanceName) {}
+
+	virtual UE::Online::EOnlineServices GetServicesProvider() const override { return UE::Online::EOnlineServices::Null; }
+
+	virtual void RegisterComponents() override;
+};
+
+class FOnlineComponentTest : public UE::Online::TOnlineComponent<UE::Online::Test::ITestInterface>
+{
+public:
+	using Super = UE::Online::TOnlineComponent<UE::Online::Test::ITestInterface>;
+
+	FOnlineComponentTest(FOnlineServicesTest& InOwningSubsystem) : Super(TEXT("Test"), InOwningSubsystem) {}
+
+	void FN(int, const FString&) {}
+	void FN2(int, const FString&, float) {}
+	void CFN(int, const FString&) const {}
+	void CFN2(int, const FString&, float) const {}
+};
+
+void FOnlineServicesTest::RegisterComponents()
+{
+	Components.Register<FOnlineComponentTest>(*this);
+}
+
 inline void DelegateTest()
 {
 	using UE::Online::Private::ConstructDelegate;
@@ -265,6 +295,10 @@ inline void DelegateTest()
 	const UMyObject* ConstUObject = UObject;
 #endif // TEST_CONSTRUCT_DELEGATE_SYNTAX_UOBJECT
 
+	TSharedRef<FOnlineServicesTest> ServicesTest = MakeShared<FOnlineServicesTest>(NAME_None);
+	FOnlineComponentTest* TestComponent = ServicesTest->Get<FOnlineComponentTest>();
+	const FOnlineComponentTest* ConstTestComponent = ServicesTest->Get<FOnlineComponentTest>();
+
 	// CreateStatic
 	TDelegate<void(int, const FString&)> StaticDelegate1 = ConstructDelegate<void(int, const FString&)>(FN);
 	TDelegate<void(int, const FString&)> StaticDelegate2 = ConstructDelegate<void(int, const FString&)>(&FN);
@@ -303,6 +337,12 @@ inline void DelegateTest()
 	TDelegate<void(int, const FString&)> TSSPRawDelegate4 = ConstructDelegate<void(int, const FString&)>(TSSPRaw, &FRawTest::CFN2, 1.0f);
 	TDelegate<void(int, const FString&)> TSSPRawDelegate5 = ConstructDelegate<void(int, const FString&)>(ConstTSSPRaw, &FRawTest::CFN);
 	TDelegate<void(int, const FString&)> TSSPRawDelegate6 = ConstructDelegate<void(int, const FString&)>(ConstTSSPRaw, &FRawTest::CFN2, 1.0f);
+	TDelegate<void(int, const FString&)> TComponentDelegate1 = ConstructDelegate<void(int, const FString&)>(TestComponent, &FOnlineComponentTest::FN);
+	TDelegate<void(int, const FString&)> TComponentDelegate2 = ConstructDelegate<void(int, const FString&)>(TestComponent, &FOnlineComponentTest::FN2, 1.0f);
+	TDelegate<void(int, const FString&)> TComponentDelegate3 = ConstructDelegate<void(int, const FString&)>(TestComponent, &FOnlineComponentTest::CFN);
+	TDelegate<void(int, const FString&)> TComponentDelegate4 = ConstructDelegate<void(int, const FString&)>(TestComponent, &FOnlineComponentTest::CFN2, 1.0f);
+	TDelegate<void(int, const FString&)> TComponentDelegate5 = ConstructDelegate<void(int, const FString&)>(ConstTestComponent, &FOnlineComponentTest::CFN);
+	TDelegate<void(int, const FString&)> TComponentDelegate6 = ConstructDelegate<void(int, const FString&)>(ConstTestComponent, &FOnlineComponentTest::CFN2, 1.0f);
 	// CreateSP
 	TDelegate<void(int, const FString&)> SPDelegate1 = ConstructDelegate<void(int, const FString&)>(SP, &FSPTest::FN);
 	TDelegate<void(int, const FString&)> SPDelegate2 = ConstructDelegate<void(int, const FString&)>(SP, &FSPTest::FN2, 1.0f);
@@ -341,6 +381,7 @@ inline void DelegateTest()
 	TDelegate<void(int, const FString&)> WeakSPDelegate2 = ConstructDelegate<void(int, const FString&)>(&SP.Get(), [](int, const FString&) {});
 	TDelegate<void(int, const FString&)> WeakTSSPRawDelegate1 = ConstructDelegate<void(int, const FString&)>(TSSPRaw, [](int, const FString&) {});
 	TDelegate<void(int, const FString&)> WeakSPRawDelegate1 = ConstructDelegate<void(int, const FString&)>(SPRaw, [](int, const FString&) {});
+	TDelegate<void(int, const FString&)> WeakTComponentDelegate = ConstructDelegate<void(int, const FString&)>(TestComponent, [](int, const FString&) {});
 	// lambda with raw pointer: Error: When using non pointer to member functions, the first parameter can only be a UObject*, TSharedRef, or pointer to a class that derives from TSharedFromThis
 //	TDelegate<void(int, const FString&)> WeakRawDelegate = ConstructDelegate<void(int, const FString&)>(Raw, [](int, const FString&) {});
 
diff --git a/Engine/Plugins/Online/OnlineServices/Source/OnlineServicesInterface/Public/Online/OnlineAsyncOpHandle.h b/Engine/Plugins/Online/OnlineServices/Source/OnlineServicesInterface/Public/Online/OnlineAsyncOpHandle.h
--- a/Engine/Plugins/Online/OnlineServices/Source/OnlineServicesInterface/Public/Online/OnlineAsyncOpHandle.h
+++ b/Engine/Plugins/Online/OnlineServices/Source/OnlineServicesInterface/Public/Online/OnlineAsyncOpHandle.h
@@ -53,7 +53,7 @@ namespace Private
 	template <typename T>
 	struct TSharedPtrHelper2
 	{
-		using PtrType = decltype(std::remove_pointer_t<T>().AsShared());
+		using PtrType = decltype(((std::remove_pointer_t<T>*)nullptr)->AsShared());
 		static inline constexpr bool bIsThreadSafe = TSharedPtrHelper<PtrType>::bIsThreadSafe;
 		static inline constexpr bool bIsSharedPtr = true;
 		using WeakPtrType = typename TSharedPtrHelper<PtrType>::WeakPtrType;
