- Fix incorrect error code masks

diff --git a/Engine/Plugins/Online/OnlineServices/Source/OnlineServicesInterface/Public/Online/OnlineErrorCode.h b/Engine/Plugins/Online/OnlineServices/Source/OnlineServicesInterface/Public/Online/OnlineErrorCode.h
--- a/Engine/Plugins/Online/OnlineServices/Source/OnlineServicesInterface/Public/Online/OnlineErrorCode.h
+++ b/Engine/Plugins/Online/OnlineServices/Source/OnlineServicesInterface/Public/Online/OnlineErrorCode.h
@@ -20,12 +20,12 @@ namespace ErrorCode
 
 	constexpr ErrorCodeType Create(uint32 Source, uint32 Category, uint32 Code)
 	{
-		return ((Source & 0x3ull) << 60ull) | ((Category & 0x3fffffffull) << 32ull) | static_cast<ErrorCodeType>(Code);
+		return ((Source & 0xfull) << 60ull) | ((Category & 0x0fffffffull) << 32ull) | static_cast<ErrorCodeType>(Code);
 	}
 
 	constexpr ErrorCodeType Create(uint32 Category, uint32 Code)
 	{
-		return ((Category & 0x3fffffffull) << 32ull) | static_cast<ErrorCodeType>(Code);
+		return ((Category & 0x0fffffffull) << 32ull) | static_cast<ErrorCodeType>(Code);
 	}
 
 	namespace System
