Minor V2 Sessions refactoring

diff --git a/Engine/Plugins/Online/OnlineServices/Source/OnlineServicesCommon/Private/Online/SessionsCommon.cpp b/Engine/Plugins/Online/OnlineServices/Source/OnlineServicesCommon/Private/Online/SessionsCommon.cpp
--- a/Engine/Plugins/Online/OnlineServices/Source/OnlineServicesCommon/Private/Online/SessionsCommon.cpp
+++ b/Engine/Plugins/Online/OnlineServices/Source/OnlineServicesCommon/Private/Online/SessionsCommon.cpp
@@ -144,14 +144,14 @@ namespace UE::Online {
 #define CHECK_PARAMS_ID_HANDLE(IdHandle, MethodName) \
 	if (!IdHandle.IsValid()) \
 	{ \
-		UE_LOG(LogTemp, Warning, TEXT("[%s] Value of [%s] invalid."), UTF8_TO_TCHAR(__FUNCTION__), UTF8_TO_TCHAR(#IdHandle)); \
+		UE_LOG(LogTemp, Verbose, TEXT("[%s] Value of [%s] invalid."), UTF8_TO_TCHAR(__FUNCTION__), UTF8_TO_TCHAR(#IdHandle)); \
 		return TOnlineResult<MethodName>(Errors::InvalidParams()); \
 	} \
 
 #define CHECK_PARAMS_FNAME(Name, MethodName) \
 	if (Name.IsNone()) \
 	{ \
-		UE_LOG(LogTemp, Warning, TEXT("[%s] Value of [%s] invalid."), UTF8_TO_TCHAR(__FUNCTION__), UTF8_TO_TCHAR(#Name)); \
+		UE_LOG(LogTemp, Verbose, TEXT("[%s] Value of [%s] invalid."), UTF8_TO_TCHAR(__FUNCTION__), UTF8_TO_TCHAR(#Name)); \
 		return TOnlineResult<MethodName>(Errors::InvalidParams()); \
 	} \
 
@@ -211,7 +211,7 @@ namespace UE::Online {
 		}
 		else
 		{
-			UE_LOG(LogTemp, Warning, TEXT("[%s] No data found for user [%s]."), UTF8_TO_TCHAR(__FUNCTION__), *ToLogString(Params.LocalAccountId));
+			UE_LOG(LogTemp, Verbose, TEXT("[%s] No data found for user [%s]."), UTF8_TO_TCHAR(__FUNCTION__), *ToLogString(Params.LocalAccountId));
 			return TOnlineResult<FGetPresenceSession>({ Errors::InvalidState() });
 		}
 	}
@@ -227,7 +227,7 @@ namespace UE::Online {
 		}
 		else
 		{
-			UE_LOG(LogTemp, Warning, TEXT("[%s] No data found for user [%s]."), UTF8_TO_TCHAR(__FUNCTION__), *ToLogString(Params.LocalAccountId));
+			UE_LOG(LogTemp, Verbose, TEXT("[%s] No data found for user [%s]."), UTF8_TO_TCHAR(__FUNCTION__), *ToLogString(Params.LocalAccountId));
 			return TOnlineResult<FIsPresenceSession>({ Errors::InvalidState() });
 		}
 	}
@@ -256,17 +256,17 @@ namespace UE::Online {
 	{
 		CHECK_PARAMS_FNAME(Params.LocalName, FGetMutableSessionByName)
 
-			if (const FOnlineSessionId* SessionId = LocalSessionsByName.Find(Params.LocalName))
-			{
-				check(AllSessionsById.Contains(*SessionId));
+		if (const FOnlineSessionId* SessionId = LocalSessionsByName.Find(Params.LocalName))
+		{
+			check(AllSessionsById.Contains(*SessionId));
 
-				return TOnlineResult<FGetMutableSessionByName>({ AllSessionsById.FindChecked(*SessionId) });
-			}
-			else
-			{
-				UE_LOG(LogTemp, Warning, TEXT("[%s] No session found for name [%s]."), UTF8_TO_TCHAR(__FUNCTION__), *Params.LocalName.ToString());
-				return TOnlineResult<FGetMutableSessionByName>(Errors::NotFound());
-			}
+			return TOnlineResult<FGetMutableSessionByName>({ AllSessionsById.FindChecked(*SessionId) });
+		}
+		else
+		{
+			UE_LOG(LogTemp, Verbose, TEXT("[%s] No session found for name [%s]."), UTF8_TO_TCHAR(__FUNCTION__), *Params.LocalName.ToString());
+			return TOnlineResult<FGetMutableSessionByName>(Errors::NotFound());
+		}
 	}
 
 	TOnlineResult<FGetMutableSessionById> FSessionsCommon::GetMutableSessionById(FGetMutableSessionById::Params&& Params) const
@@ -279,7 +279,7 @@ namespace UE::Online {
 		}
 		else
 		{
-			UE_LOG(LogTemp, Warning, TEXT("[%s] No session found for id [%s]."), UTF8_TO_TCHAR(__FUNCTION__), *ToLogString(Params.SessionId));
+			UE_LOG(LogTemp, Verbose, TEXT("[%s] No session found for id [%s]."), UTF8_TO_TCHAR(__FUNCTION__), *ToLogString(Params.SessionId));
 			return TOnlineResult<FGetMutableSessionById>(Errors::NotFound());
 		}
 	}
@@ -297,13 +297,13 @@ namespace UE::Online {
 			}
 			else
 			{
-				UE_LOG(LogTemp, Warning, TEXT("[%s] No invite found for invite id [%s]."), UTF8_TO_TCHAR(__FUNCTION__), *ToLogString(Params.SessionInviteId));
+				UE_LOG(LogTemp, Verbose, TEXT("[%s] No invite found for invite id [%s]."), UTF8_TO_TCHAR(__FUNCTION__), *ToLogString(Params.SessionInviteId));
 				return TOnlineResult<FGetSessionInviteById>(Errors::NotFound());
 			}
 		}
 		else
 		{
-			UE_LOG(LogTemp, Warning, TEXT("[%s] No invite data found for account id [%s]."), UTF8_TO_TCHAR(__FUNCTION__), *ToLogString(Params.LocalAccountId));
+			UE_LOG(LogTemp, Verbose, TEXT("[%s] No invite data found for account id [%s]."), UTF8_TO_TCHAR(__FUNCTION__), *ToLogString(Params.LocalAccountId));
 			return TOnlineResult<FGetSessionInviteById>(Errors::InvalidState());
 		}
 	}
@@ -333,30 +333,22 @@ namespace UE::Online {
 #define CHECK_PARAMS_ID_HANDLE(IdHandle) \
 	if (!IdHandle.IsValid()) \
 	{ \
-		UE_LOG(LogTemp, Warning, TEXT("[%s] Value of [%s] invalid."), UTF8_TO_TCHAR(__FUNCTION__), UTF8_TO_TCHAR(#IdHandle)); \
+		UE_LOG(LogTemp, Verbose, TEXT("[%s] Value of [%s] invalid."), UTF8_TO_TCHAR(__FUNCTION__), UTF8_TO_TCHAR(#IdHandle)); \
 		return TOptional<FOnlineError>(Errors::InvalidParams()); \
 	} \
 
 #define CHECK_PARAMS_FNAME(Name) \
 	if (Name.IsNone()) \
 	{ \
-		UE_LOG(LogTemp, Warning, TEXT("[%s] Value of [%s] invalid."), UTF8_TO_TCHAR(__FUNCTION__), UTF8_TO_TCHAR(#Name)); \
+		UE_LOG(LogTemp, Verbose, TEXT("[%s] Value of [%s] invalid."), UTF8_TO_TCHAR(__FUNCTION__), UTF8_TO_TCHAR(#Name)); \
 		return TOptional<FOnlineError>(Errors::InvalidParams()); \
 	} \
 
-#define CHECK_STATE_NEW_SESSION_NAME(SessionName) \
-	TOnlineResult<FGetSessionByName> GetSessionByNameResult = GetSessionByName({ SessionName }); \
-	if (GetSessionByNameResult.IsOk()) \
-	{ \
-		UE_LOG(LogTemp, Warning, TEXT("[%s] Session with name [%s] already exists."), UTF8_TO_TCHAR(__FUNCTION__), *SessionName.ToString()); \
-		return TOptional<FOnlineError>(Errors::InvalidState()); \
-	} \
-
 #define CHECK_STATE_EXISTING_SESSION_NAME(SessionName) \
 	TOnlineResult<FGetSessionByName> GetSessionByNameResult = GetSessionByName({ SessionName }); \
 	if (GetSessionByNameResult.IsError()) \
 	{ \
-		UE_LOG(LogTemp, Warning, TEXT("[%s] Could not find session with name [%s]."), UTF8_TO_TCHAR(__FUNCTION__), *SessionName.ToString()); \
+		UE_LOG(LogTemp, Verbose, TEXT("[%s] Could not find session with name [%s]."), UTF8_TO_TCHAR(__FUNCTION__), *SessionName.ToString()); \
 		return TOptional<FOnlineError>(Errors::InvalidState()); \
 	} \
 
@@ -426,7 +418,7 @@ namespace UE::Online {
 
 		if (Params.SessionSettings.NumMaxConnections == 0)
 		{
-			UE_LOG(LogTemp, Warning, TEXT("[%s] NumMaxConnections invalid, can't be 0."), UTF8_TO_TCHAR(__FUNCTION__));
+			UE_LOG(LogTemp, Verbose, TEXT("[%s] NumMaxConnections invalid, can't be 0."), UTF8_TO_TCHAR(__FUNCTION__));
 			return Errors::InvalidParams();
 		}
 
@@ -442,7 +434,12 @@ namespace UE::Online {
 
 	TOptional<FOnlineError> FSessionsCommon::CheckState(const FCreateSession::Params& Params) const
 	{
-		CHECK_STATE_NEW_SESSION_NAME(Params.SessionName)
+		TOnlineResult<FGetSessionByName> GetSessionByNameResult = GetSessionByName({ Params.SessionName });
+		if (GetSessionByNameResult.IsOk())
+		{
+			UE_LOG(LogTemp, Verbose, TEXT("[%s] Another session with the name [%s] already exists."), UTF8_TO_TCHAR(__FUNCTION__), *Params.SessionName.ToString());
+			return TOptional<FOnlineError>(Errors::InvalidState());
+		}
 
 		if (Params.bPresenceEnabled)
 		{
@@ -451,7 +448,7 @@ namespace UE::Online {
 				TOnlineResult<FGetPresenceSession> GetPresenceSessionResult = GetPresenceSession({ Params.LocalAccountId });
 				if (GetPresenceSessionResult.IsOk())
 				{
-					UE_LOG(LogTemp, Warning, TEXT("[%s] Could not create session with bPresenceEnabled set to true when another already exists [%s]."), UTF8_TO_TCHAR(__FUNCTION__), *Entry.Key.ToString());
+					UE_LOG(LogTemp, Verbose, TEXT("[%s] Could not create session with bPresenceEnabled set to true when another already exists [%s]."), UTF8_TO_TCHAR(__FUNCTION__), *Entry.Key.ToString());
 					return Errors::InvalidState();
 				}
 			}
@@ -486,7 +483,7 @@ namespace UE::Online {
 		{
 			if (Params.Mutations.NumMaxConnections.GetValue() == 0)
 			{
-				UE_LOG(LogTemp, Warning, TEXT("[%s] NumMaxConnections invalid, can't be 0."), UTF8_TO_TCHAR(__FUNCTION__));
+				UE_LOG(LogTemp, Verbose, TEXT("[%s] NumMaxConnections invalid, can't be 0."), UTF8_TO_TCHAR(__FUNCTION__));
 				return Errors::InvalidParams();
 			}
 		}
@@ -558,7 +555,7 @@ namespace UE::Online {
 
 		if (Params.MaxResults == 0)
 		{
-			UE_LOG(LogTemp, Warning, TEXT("[%s] MaxResults can't be 0"), UTF8_TO_TCHAR(__FUNCTION__));
+			UE_LOG(LogTemp, Verbose, TEXT("[%s] MaxResults can't be 0"), UTF8_TO_TCHAR(__FUNCTION__));
 			return TOptional<FOnlineError>(Errors::InvalidParams());
 		}
 
@@ -584,7 +581,7 @@ namespace UE::Online {
 	{
 		if (CurrentSessionSearchPromisesUserMap.Contains(Params.LocalAccountId))
 		{
-			UE_LOG(LogTemp, Warning, TEXT("[%s] Search already in progress"), UTF8_TO_TCHAR(__FUNCTION__));
+			UE_LOG(LogTemp, Verbose, TEXT("[%s] Search already in progress"), UTF8_TO_TCHAR(__FUNCTION__));
 			return Errors::AlreadyPending();
 		}
 
@@ -651,12 +648,20 @@ namespace UE::Online {
 
 	TOptional<FOnlineError> FSessionsCommon::CheckState(const FJoinSession::Params& Params) const
 	{
-		CHECK_STATE_NEW_SESSION_NAME(Params.SessionName)
+		TOnlineResult<FGetSessionByName> GetSessionByNameResult = GetSessionByName({ Params.SessionName });
+		if (GetSessionByNameResult.IsOk())
+		{
+			if (GetSessionByNameResult.GetOkValue().Session->GetSessionId() != Params.SessionId)
+			{
+				UE_LOG(LogTemp, Verbose, TEXT("[%s] Another session with the name [%s] already exists."), UTF8_TO_TCHAR(__FUNCTION__), *Params.SessionName.ToString());
+				return TOptional<FOnlineError>(Errors::InvalidState());
+			}
+		}
 
 		TOnlineResult<FGetSessionById> GetSessionByIdResult = GetSessionById({ Params.SessionId });
 		if (GetSessionByIdResult.IsError())
 		{
-			UE_LOG(LogTemp, Warning, TEXT("[%s] Session [%s] not found. Please call FindSessions to get an updated list of available sessions "), UTF8_TO_TCHAR(__FUNCTION__), *ToLogString(Params.SessionId));
+			UE_LOG(LogTemp, Verbose, TEXT("[%s] Session [%s] not found. Please call FindSessions to get an updated list of available sessions "), UTF8_TO_TCHAR(__FUNCTION__), *ToLogString(Params.SessionId));
 			return TOptional<FOnlineError>(GetSessionByIdResult.GetErrorValue());
 		}
 
@@ -666,13 +671,13 @@ namespace UE::Online {
 
 		if (FoundSession->GetSessionMembers().Contains(Params.LocalAccountId))
 		{
-			UE_LOG(LogTemp, Warning, TEXT("[%s] User [%s] already in session [%s]"), UTF8_TO_TCHAR(__FUNCTION__), *ToLogString(Params.LocalAccountId), *ToLogString(Params.SessionId));
+			UE_LOG(LogTemp, Verbose, TEXT("[%s] User [%s] already in session [%s]"), UTF8_TO_TCHAR(__FUNCTION__), *ToLogString(Params.LocalAccountId), *ToLogString(Params.SessionId));
 			return TOptional<FOnlineError>(Errors::AccessDenied());
 		}
 
 		if (!FoundSession->IsJoinable())
 		{
-			UE_LOG(LogTemp, Warning, TEXT("[%s] Session [%s] not joinable "), UTF8_TO_TCHAR(__FUNCTION__), *ToLogString(Params.SessionId));
+			UE_LOG(LogTemp, Verbose, TEXT("[%s] Session [%s] not joinable "), UTF8_TO_TCHAR(__FUNCTION__), *ToLogString(Params.SessionId));
 			return TOptional<FOnlineError>(Errors::AccessDenied());
 		}
 
@@ -683,7 +688,7 @@ namespace UE::Online {
 				TOnlineResult<FGetPresenceSession> GetPresenceSessionResult = GetPresenceSession({ Params.LocalAccountId });
 				if (GetPresenceSessionResult.IsOk())
 				{
-					UE_LOG(LogTemp, Warning, TEXT("[%s] Could not join session with bPresenceEnabled set to true when another presence session already exists [%s]."), UTF8_TO_TCHAR(__FUNCTION__), *Entry.Key.ToString());
+					UE_LOG(LogTemp, Verbose, TEXT("[%s] Could not join session with bPresenceEnabled set to true when another presence session already exists [%s]."), UTF8_TO_TCHAR(__FUNCTION__), *Entry.Key.ToString());
 					return TOptional<FOnlineError>(Errors::InvalidState());
 				}
 			}
@@ -809,7 +814,7 @@ namespace UE::Online {
 
 		if (Params.TargetUsers.IsEmpty())
 		{
-			UE_LOG(LogTemp, Warning, TEXT("[%s] TargetUsers array is empty"), UTF8_TO_TCHAR(__FUNCTION__));
+			UE_LOG(LogTemp, Verbose, TEXT("[%s] TargetUsers array is empty"), UTF8_TO_TCHAR(__FUNCTION__));
 			return TOptional<FOnlineError>(Errors::InvalidParams());
 		}
 
@@ -854,13 +859,13 @@ namespace UE::Online {
 		{
 			if (!UserMap->Contains(Params.SessionInviteId))
 			{
-				UE_LOG(LogTemp, Warning, TEXT("[%s] No invite found for invite id [%s]."), UTF8_TO_TCHAR(__FUNCTION__), *ToLogString(Params.SessionInviteId));
+				UE_LOG(LogTemp, Verbose, TEXT("[%s] No invite found for invite id [%s]."), UTF8_TO_TCHAR(__FUNCTION__), *ToLogString(Params.SessionInviteId));
 				return TOptional<FOnlineError>(Errors::NotFound());
 			}
 		}
 		else
 		{
-			UE_LOG(LogTemp, Warning, TEXT("[%s] No invite data found for account id [%s]."), UTF8_TO_TCHAR(__FUNCTION__), *ToLogString(Params.LocalAccountId));
+			UE_LOG(LogTemp, Verbose, TEXT("[%s] No invite data found for account id [%s]."), UTF8_TO_TCHAR(__FUNCTION__), *ToLogString(Params.LocalAccountId));
 			return TOptional<FOnlineError>(Errors::InvalidState());
 		}
 
