Added option for sessions methods to use the parallel async queue

diff --git a/Engine/Plugins/Online/OnlineServices/Source/OnlineServicesCommon/Private/Online/SessionsCommon.cpp b/Engine/Plugins/Online/OnlineServices/Source/OnlineServicesCommon/Private/Online/SessionsCommon.cpp
--- a/Engine/Plugins/Online/OnlineServices/Source/OnlineServicesCommon/Private/Online/SessionsCommon.cpp
+++ b/Engine/Plugins/Online/OnlineServices/Source/OnlineServicesCommon/Private/Online/SessionsCommon.cpp
@@ -353,7 +353,7 @@ namespace UE::Online {
 	} \
 
 	template<typename MethodStruct>
-	TOnlineAsyncOpHandle<MethodStruct> FSessionsCommon::ExecuteAsyncSessionsMethod(typename MethodStruct::Params&& Params, TFuture<TOnlineResult<MethodStruct>>(FSessionsCommon::* ImplFunc)(const typename MethodStruct::Params& Params))
+	TOnlineAsyncOpHandle<MethodStruct> FSessionsCommon::ExecuteAsyncSessionsMethod(typename MethodStruct::Params&& Params, TFuture<TOnlineResult<MethodStruct>>(FSessionsCommon::* ImplFunc)(const typename MethodStruct::Params& Params), bool bUseParallelQueue)
 	{
 		TOnlineAsyncOpRef<MethodStruct> Op = GetOp<MethodStruct>(MoveTemp(Params));
 		const typename MethodStruct::Params& OpParams = Op->GetParams();
@@ -393,7 +393,7 @@ namespace UE::Online {
 							});
 				}
 			})
-			.Enqueue(GetSerialQueue());
+			.Enqueue(bUseParallelQueue ? Services.GetParallelQueue() : GetSerialQueue());
 
 		return Op->GetHandle();
 	}
@@ -541,7 +541,7 @@ namespace UE::Online {
 
 	TOnlineAsyncOpHandle<FFindSessions> FSessionsCommon::FindSessions(FFindSessions::Params&& Params)
 	{
-		return ExecuteAsyncSessionsMethod<FFindSessions>(MoveTemp(Params), &FSessionsCommon::FindSessionsImpl);
+		return ExecuteAsyncSessionsMethod<FFindSessions>(MoveTemp(Params), &FSessionsCommon::FindSessionsImpl, true);
 	}
 
 	TFuture<TOnlineResult<FFindSessions>> FSessionsCommon::FindSessionsImpl(const FFindSessions::Params& Params)
diff --git a/Engine/Plugins/Online/OnlineServices/Source/OnlineServicesCommon/Public/Online/SessionsCommon.h b/Engine/Plugins/Online/OnlineServices/Source/OnlineServicesCommon/Public/Online/SessionsCommon.h
--- a/Engine/Plugins/Online/OnlineServices/Source/OnlineServicesCommon/Public/Online/SessionsCommon.h
+++ b/Engine/Plugins/Online/OnlineServices/Source/OnlineServicesCommon/Public/Online/SessionsCommon.h
@@ -328,7 +328,7 @@ protected:
 
 private:
 	template<typename MethodStruct>
-	TOnlineAsyncOpHandle<MethodStruct> ExecuteAsyncSessionsMethod(typename MethodStruct::Params&& Params, TFuture<TOnlineResult<MethodStruct>>(FSessionsCommon::* ImplFunc)(const typename MethodStruct::Params& Params));
+	TOnlineAsyncOpHandle<MethodStruct> ExecuteAsyncSessionsMethod(typename MethodStruct::Params&& Params, TFuture<TOnlineResult<MethodStruct>>(FSessionsCommon::* ImplFunc)(const typename MethodStruct::Params& Params), bool bUseParallelQueue = false);
 
 	void ClearSessionByName(const FName& SessionName);
 	void ClearSessionById(const FOnlineSessionId& SessionId);
