Fix issue causing non-deterministic OSSv2 SchemaCompatibilityIds.
Set local name in EOS lobby on create and join.
Fix lyra lobby schema config issue.

diff --git a/Engine/Plugins/Online/OnlineServices/Source/OnlineServicesInterface/Private/Online/Schema.cpp b/Engine/Plugins/Online/OnlineServices/Source/OnlineServicesInterface/Private/Online/Schema.cpp
--- a/Engine/Plugins/Online/OnlineServices/Source/OnlineServicesInterface/Private/Online/Schema.cpp
+++ b/Engine/Plugins/Online/OnlineServices/Source/OnlineServicesInterface/Private/Online/Schema.cpp
@@ -20,6 +20,12 @@ T BuildFlagsFromArray(const TArray<T>& Input)
 	return OutValue;
 }
 
+static uint32 SchemaNameCrc32(const FName& Name, uint32 CRC=0)
+{
+	const FString NameString = Name.ToString().ToLower();
+	return FCrc::MemCrc32(GetData(NameString), GetNum(NameString), CRC);
+}
+
 ESchemaServiceAttributeSupportedTypeFlags TranslateAttributeType(ESchemaAttributeType InType)
 {
 	switch (InType)
@@ -553,7 +559,7 @@ bool FSchemaRegistry::ParseConfig(const FSchemaRegistryDescriptorConfig& Config)
 
 				// Populate parent info.
 				SchemaDefinition->ParentSchemaIds.Add(IterSchemaId);
-				SchemaDataCrc = FCrc::TypeCrc32(IterSchemaId, SchemaDataCrc);
+				SchemaDataCrc = SchemaNameCrc32(IterSchemaId, SchemaDataCrc);
 			}
 			Algo::Reverse(SchemaHierarchy);
 			check(SchemaHierarchy.Num() > 0);
@@ -566,7 +572,7 @@ bool FSchemaRegistry::ParseConfig(const FSchemaRegistryDescriptorConfig& Config)
 				FSchemaCategoryDefinition& SchemaCategoryDefinition = SchemaDefinition->Categories.Emplace(CategoryId);
 				SchemaCategoryDefinition.Id = CategoryId;
 
-				SchemaDataCrc = FCrc::TypeCrc32(CategoryId, SchemaDataCrc);
+				SchemaDataCrc = SchemaNameCrc32(CategoryId, SchemaDataCrc);
 			}
 
 			// Iterate the hierarchy from root -> leaf
@@ -601,7 +607,7 @@ bool FSchemaRegistry::ParseConfig(const FSchemaRegistryDescriptorConfig& Config)
 									SchemaCategoryDefinition.SchemaCompatibilityAttributeId = AttributeDefinition.Id;
 								}
 
-								SchemaDataCrc = FCrc::TypeCrc32(AttributeDefinition.Id, SchemaDataCrc);
+								SchemaDataCrc = SchemaNameCrc32(AttributeDefinition.Id, SchemaDataCrc);
 								SchemaDataCrc = FCrc::TypeCrc32(AttributeDefinition.Flags, SchemaDataCrc);
 								SchemaDataCrc = FCrc::TypeCrc32(AttributeDefinition.Type, SchemaDataCrc);
 								SchemaDataCrc = FCrc::TypeCrc32(AttributeDefinition.MaxSize, SchemaDataCrc);
@@ -708,12 +714,12 @@ bool FSchemaRegistry::ParseConfig(const FSchemaRegistryDescriptorConfig& Config)
 								}
 
 								// Add service attribute definition to data crc.
-								SchemaDataCrc = FCrc::TypeCrc32(SchemaServiceAttributeDefinition.Id, SchemaDataCrc);
+								SchemaDataCrc = SchemaNameCrc32(SchemaServiceAttributeDefinition.Id, SchemaDataCrc);
 								SchemaDataCrc = FCrc::TypeCrc32(SchemaServiceAttributeDefinition.Type, SchemaDataCrc);
 								SchemaDataCrc = FCrc::TypeCrc32(SchemaServiceAttributeDefinition.Flags, SchemaDataCrc);
 								SchemaDataCrc = FCrc::TypeCrc32(SchemaServiceAttributeDefinition.MaxSize, SchemaDataCrc);
 
-								Algo::ForEach(SchemaServiceAttributeDefinition.SchemaAttributeIds, [&SchemaDataCrc](const FSchemaAttributeId& Id) { SchemaDataCrc = FCrc::TypeCrc32(Id, SchemaDataCrc); });
+								Algo::ForEach(SchemaServiceAttributeDefinition.SchemaAttributeIds, [&SchemaDataCrc](const FSchemaAttributeId& Id) { SchemaDataCrc = SchemaNameCrc32(Id, SchemaDataCrc); });
 							}
 							else
 							{
@@ -726,14 +732,14 @@ bool FSchemaRegistry::ParseConfig(const FSchemaRegistryDescriptorConfig& Config)
 						}
 					}
 
-					SchemaDataCrc = FCrc::TypeCrc32(SchemaCategoryDefinition.SchemaCompatibilityAttributeId, SchemaDataCrc);
-					SchemaDataCrc = FCrc::TypeCrc32(SchemaCategoryDefinition.SchemaCompatibilityServiceAttributeId, SchemaDataCrc);
+					SchemaDataCrc = SchemaNameCrc32(SchemaCategoryDefinition.SchemaCompatibilityAttributeId, SchemaDataCrc);
+					SchemaDataCrc = SchemaNameCrc32(SchemaCategoryDefinition.SchemaCompatibilityServiceAttributeId, SchemaDataCrc);
 				}
 			}
 
 			// Build compatibility id from the schema and data CRCs.
 			{
-				const uint32 SchemaIdCrc = FCrc::TypeCrc32(SchemaDefinition->Id);
+				const uint32 SchemaIdCrc = SchemaNameCrc32(SchemaDefinition->Id);
 				if (SeenSchemaCrcs.Find(SchemaIdCrc) != nullptr)
 				{
 					UE_LOG(LogOnlineSchema, Error, TEXT("Invalid schema %s: CRC collision processing schema name, please rename the schema to avoid collision."),
@@ -744,9 +750,15 @@ bool FSchemaRegistry::ParseConfig(const FSchemaRegistryDescriptorConfig& Config)
 
 				const uint64 CompatibilityIdUint = (static_cast<uint64>(SchemaIdCrc) << 32) | SchemaDataCrc;
 				SchemaDefinition->CompatibilityId = *reinterpret_cast<const int64*>(&CompatibilityIdUint);
+
+				UE_LOG(LogOnlineSchema, Verbose, TEXT("Building compatability id: SchemaId: %s, SchemaIdCrc: 0x%04X, SchemaDataCrc 0x%04X."),
+					*SchemaDefinition->Id.ToString().ToLower(), SchemaIdCrc, SchemaDataCrc);
 			}
 
 			ParsedSchemaDefinitionsByCompatibilityId.Add(SchemaDefinition->CompatibilityId, SchemaDefinition);
+
+			UE_LOG(LogOnlineSchema, Log, TEXT("Parsed Schema definition: SchemaId: %s, SchemaCompatabilityId: 0x%08" INT64_X_FMT "."),
+				*SchemaDefinition->Id.ToString().ToLower(), SchemaDefinition->CompatibilityId);
 		}
 	}
 
@@ -1199,11 +1211,11 @@ TOnlineResult<FSchemaCategoryInstancePrepareServiceSnapshot> FSchemaCategoryInst
 				// Schema changes are handled by surfacing a new schema id in the result object. Don't return SchemaCompatibilityId as an attribute to the client.
 				if (EnumHasAnyFlags(SchemaAttributeDefinition->Flags, ESchemaAttributeFlags::SchemaCompatibilityId))
 				{
-					UE_LOG(LogOnlineSchema, Verbose, TEXT("[PrepareServiceSnapshot] %s.%s.%s: Consuming SchemaCompatibility attribute data: %s."),
+					UE_LOG(LogOnlineSchema, Verbose, TEXT("[PrepareServiceSnapshot] %s.%s.%s: Consuming SchemaCompatibility attribute: 0x%08" INT64_X_FMT "."),
 						*SchemaDefinition->Id.ToString().ToLower(),
 						*SchemaCategoryDefinition->Id.ToString().ToLower(),
 						*SchemaAttributeDefinition->Id.ToString().ToLower(),
-						*SchemaServiceAttributeValue.ToLogString());
+						SchemaServiceAttributeValue.GetInt64());
 				}
 				else
 				{
diff --git a/Engine/Plugins/Online/OnlineServicesEOSGS/Source/Private/Online/LobbiesEOSGS.cpp b/Engine/Plugins/Online/OnlineServicesEOSGS/Source/Private/Online/LobbiesEOSGS.cpp
--- a/Engine/Plugins/Online/OnlineServicesEOSGS/Source/Private/Online/LobbiesEOSGS.cpp
+++ b/Engine/Plugins/Online/OnlineServicesEOSGS/Source/Private/Online/LobbiesEOSGS.cpp
@@ -218,6 +218,7 @@ TOnlineAsyncOpHandle<FCreateLobby> FLobbiesEOSGS::CreateLobby(FCreateLobby::Para
 		PrepareParams.ClientChanges.MemberAttributes = { Params.UserAttributes, {} };
 		PrepareParams.ClientChanges.JoinPolicy = Params.JoinPolicy;
 		PrepareParams.ClientChanges.LobbySchema = Params.SchemaId;
+		PrepareParams.ClientChanges.LocalName = Params.LocalName;
 		TOnlineResult<FLobbyClientDataPrepareClientChanges> PrepareResult = LobbyData->GetLobbyClientData()->PrepareClientChanges(MoveTemp(PrepareParams));
 		if (PrepareResult.IsError())
 		{
@@ -413,6 +414,7 @@ TOnlineAsyncOpHandle<FJoinLobby> FLobbiesEOSGS::JoinLobby(FJoinLobby::Params&& I
 
 		FLobbyClientDataPrepareClientChanges::Params PrepareParams;
 		PrepareParams.LocalAccountId = Params.LocalAccountId;
+		PrepareParams.ClientChanges.LocalName = Params.LocalName;
 		PrepareParams.ClientChanges.MemberAttributes = { Params.UserAttributes, {} };
 		TOnlineResult<FLobbyClientDataPrepareClientChanges> PrepareResult = LobbyData->GetLobbyClientData()->PrepareClientChanges(MoveTemp(PrepareParams));
 		if (PrepareResult.IsError())
