Fix issue caused by wide character hashing when creating schema compatibility id.

diff --git a/Engine/Plugins/Online/OnlineServices/Source/OnlineServicesInterface/Private/Online/Schema.cpp b/Engine/Plugins/Online/OnlineServices/Source/OnlineServicesInterface/Private/Online/Schema.cpp
--- a/Engine/Plugins/Online/OnlineServices/Source/OnlineServicesInterface/Private/Online/Schema.cpp
+++ b/Engine/Plugins/Online/OnlineServices/Source/OnlineServicesInterface/Private/Online/Schema.cpp
@@ -23,7 +23,7 @@ T BuildFlagsFromArray(const TArray<T>& Input)
 static uint32 SchemaNameCrc32(const FName& Name, uint32 CRC=0)
 {
 	const FString NameString = Name.ToString().ToLower();
-	return FCrc::MemCrc32(GetData(NameString), GetNum(NameString), CRC);
+	return FCrc::StrCrc32<TCHAR>(*NameString, CRC);
 }
 
 ESchemaServiceAttributeSupportedTypeFlags TranslateAttributeType(ESchemaAttributeType InType)
