Fix for V2 Sessions Bug:
Interface methods to operate Presence sessions were being used, which can be unimplemented.
Added Impl methods in Common layer and used them instead.

diff --git a/Engine/Plugins/Online/OnlineServices/Source/OnlineServicesCommon/Private/Online/SessionsCommon.cpp b/Engine/Plugins/Online/OnlineServices/Source/OnlineServicesCommon/Private/Online/SessionsCommon.cpp
--- a/Engine/Plugins/Online/OnlineServices/Source/OnlineServicesCommon/Private/Online/SessionsCommon.cpp
+++ b/Engine/Plugins/Online/OnlineServices/Source/OnlineServicesCommon/Private/Online/SessionsCommon.cpp
@@ -232,6 +232,16 @@ namespace UE::Online {
 	}
 
 	TOnlineResult<FSetPresenceSession> FSessionsCommon::SetPresenceSession(FSetPresenceSession::Params&& Params)
+	{
+		return SetPresenceSessionImpl(MoveTemp(Params));
+	}
+
+	TOnlineResult<FClearPresenceSession> FSessionsCommon::ClearPresenceSession(FClearPresenceSession::Params&& Params)
+	{
+		return ClearPresenceSessionImpl(MoveTemp(Params));
+	}
+
+	TOnlineResult<FSetPresenceSession> FSessionsCommon::SetPresenceSessionImpl(FSetPresenceSession::Params&& Params)
 	{
 		CHECK_PARAMS_ID_HANDLE(Params.LocalAccountId, FSetPresenceSession)
 		CHECK_PARAMS_ID_HANDLE(Params.SessionId, FSetPresenceSession);
@@ -242,7 +252,7 @@ namespace UE::Online {
 		return TOnlineResult<FSetPresenceSession>(FSetPresenceSession::Result{ });
 	}
 
-	TOnlineResult<FClearPresenceSession> FSessionsCommon::ClearPresenceSession(FClearPresenceSession::Params&& Params)
+	TOnlineResult<FClearPresenceSession> FSessionsCommon::ClearPresenceSessionImpl(FClearPresenceSession::Params&& Params)
 	{
 		CHECK_PARAMS_ID_HANDLE(Params.LocalAccountId, FClearPresenceSession);
 
@@ -935,7 +945,7 @@ namespace UE::Online {
 
 		if (bIsPresenceSession)
 		{
-			SetPresenceSession({ LocalAccountId, SessionId });
+			SetPresenceSessionImpl({ LocalAccountId, SessionId });
 		}
 	}
 
@@ -968,7 +978,7 @@ namespace UE::Online {
 		{
 			if (IsPresenceSessionResult.GetOkValue().bIsPresenceSession)
 			{
-				ClearPresenceSession({ LocalAccountId });
+				ClearPresenceSessionImpl({ LocalAccountId });
 			}
 		}
 
diff --git a/Engine/Plugins/Online/OnlineServices/Source/OnlineServicesCommon/Public/Online/SessionsCommon.h b/Engine/Plugins/Online/OnlineServices/Source/OnlineServicesCommon/Public/Online/SessionsCommon.h
--- a/Engine/Plugins/Online/OnlineServices/Source/OnlineServicesCommon/Public/Online/SessionsCommon.h
+++ b/Engine/Plugins/Online/OnlineServices/Source/OnlineServicesCommon/Public/Online/SessionsCommon.h
@@ -284,6 +284,9 @@ protected:
 
 	// FSessionsCommon internal interface
 
+	virtual TOnlineResult<FSetPresenceSession> SetPresenceSessionImpl(FSetPresenceSession::Params&& Params);
+	virtual TOnlineResult<FClearPresenceSession> ClearPresenceSessionImpl(FClearPresenceSession::Params&& Params);
+
 	virtual TFuture<TOnlineResult<FCreateSession>> CreateSessionImpl(const FCreateSession::Params& Params);
 	virtual TOptional<FOnlineError> CheckParams(const FCreateSession::Params& Params) const;
 	virtual TOptional<FOnlineError> CheckState(const FCreateSession::Params& Params) const;
