Updated unity build file sizes.

OnlineServicesCommon:
Before: 393216 | 103.984375s | 104.02s  | 1
After: 136608 | 107.375s | 54.30 | 2

MetasoundFrontend:
Before: 393216 | 126.703125 | 126.86 | 1
After: 294912 | 138.515625 | 80.14 | 3

AppFramework:
Before: 393216 | 86.28125 | 86.4 | 1
Before: 244912 | 110.296875 | 66.3 | 2

diff --git a/Engine/Plugins/Online/OnlineServices/Source/OnlineServicesCommon/OnlineServicesCommon.Build.cs b/Engine/Plugins/Online/OnlineServices/Source/OnlineServicesCommon/OnlineServicesCommon.Build.cs
--- a/Engine/Plugins/Online/OnlineServices/Source/OnlineServicesCommon/OnlineServicesCommon.Build.cs
+++ b/Engine/Plugins/Online/OnlineServices/Source/OnlineServicesCommon/OnlineServicesCommon.Build.cs
@@ -7,6 +7,7 @@ public class OnlineServicesCommon : ModuleRules
 	public OnlineServicesCommon(ReadOnlyTargetRules Target) : base(Target)
 	{
 		PCHUsage = ModuleRules.PCHUsageMode.UseExplicitOrSharedPCHs;
+		NumIncludedBytesPerUnityCPPOverride = 136608; // best unity size found from using UBT ProfileUnitySizes mode
 
 		PublicDependencyModuleNames.AddRange(
 			new string[] {
diff --git a/Engine/Plugins/Runtime/Metasound/Source/MetasoundFrontend/MetasoundFrontend.Build.cs b/Engine/Plugins/Runtime/Metasound/Source/MetasoundFrontend/MetasoundFrontend.Build.cs
--- a/Engine/Plugins/Runtime/Metasound/Source/MetasoundFrontend/MetasoundFrontend.Build.cs
+++ b/Engine/Plugins/Runtime/Metasound/Source/MetasoundFrontend/MetasoundFrontend.Build.cs
@@ -9,6 +9,8 @@ namespace UnrealBuildTool.Rules
 	{
 		public MetasoundFrontend(ReadOnlyTargetRules Target) : base(Target)
 		{
+			NumIncludedBytesPerUnityCPPOverride = 294912; // best unity size found from using UBT ProfileUnitySizes mode
+
 			PublicDependencyModuleNames.AddRange
 			(
 				new string[]
diff --git a/Engine/Source/Runtime/AppFramework/AppFramework.Build.cs b/Engine/Source/Runtime/AppFramework/AppFramework.Build.cs
--- a/Engine/Source/Runtime/AppFramework/AppFramework.Build.cs
+++ b/Engine/Source/Runtime/AppFramework/AppFramework.Build.cs
@@ -6,6 +6,8 @@ public class AppFramework : ModuleRules
 {
 	public AppFramework(ReadOnlyTargetRules Target) : base(Target)
 	{
+		NumIncludedBytesPerUnityCPPOverride = 196608; // best unity size found from using UBT ProfileUnitySizes mode
+
 		PublicDependencyModuleNames.AddRange(
 			new string[] {
 				"Core",
