Add support for BuildIdOverride on the command line to be surrounded by quotes

diff --git a/Engine/Plugins/Online/OnlineServices/Source/OnlineServicesInterface/Private/Online/OnlineServices.cpp b/Engine/Plugins/Online/OnlineServices/Source/OnlineServicesInterface/Private/Online/OnlineServices.cpp
--- a/Engine/Plugins/Online/OnlineServices/Source/OnlineServicesInterface/Private/Online/OnlineServices.cpp
+++ b/Engine/Plugins/Online/OnlineServices/Source/OnlineServicesInterface/Private/Online/OnlineServices.cpp
@@ -25,7 +25,12 @@ int32 GetBuildUniqueId()
 	if (!bStaticCheck)
 	{
 		bStaticCheck = true;
-		if (FParse::Value(FCommandLine::Get(), TEXT("BuildIdOverride="), BuildIdOverride) && BuildIdOverride != 0)
+		FString BuildIdOverrideCommandLineString;
+		if (FParse::Value(FCommandLine::Get(), TEXT("BuildIdOverride="), BuildIdOverrideCommandLineString))
+		{
+			BuildIdOverride = FCString::Atoi(*BuildIdOverrideCommandLineString);
+		}
+		if (BuildIdOverride != 0)
 		{
 			bUseBuildIdOverride = true;
 		}
diff --git a/Engine/Plugins/Online/OnlineSubsystem/Source/Private/OnlineSubsystem.cpp b/Engine/Plugins/Online/OnlineSubsystem/Source/Private/OnlineSubsystem.cpp
--- a/Engine/Plugins/Online/OnlineSubsystem/Source/Private/OnlineSubsystem.cpp
+++ b/Engine/Plugins/Online/OnlineSubsystem/Source/Private/OnlineSubsystem.cpp
@@ -149,7 +149,12 @@ int32 GetBuildUniqueId()
 	if (!bStaticCheck)
 	{
 		bStaticCheck = true;
-		if (FParse::Value(FCommandLine::Get(), TEXT("BuildIdOverride="), BuildIdOverride) && BuildIdOverride != 0)
+		FString BuildIdOverrideCommandLineString;
+		if (FParse::Value(FCommandLine::Get(), TEXT("BuildIdOverride="), BuildIdOverrideCommandLineString))
+		{
+			BuildIdOverride = FCString::Atoi(*BuildIdOverrideCommandLineString);
+		}
+		if (BuildIdOverride != 0)
 		{
 			bUseBuildIdOverride = true;
 		}
