Fix issue with OnlineTests not being able to compile with LogConsoleResponse.
 - Wraps logging WITH_ENGINE only

diff --git a/Engine/Plugins/Online/OnlineServices/Source/OnlineServicesCommon/Public/Online/OnlineExecHandler.h b/Engine/Plugins/Online/OnlineServices/Source/OnlineServicesCommon/Public/Online/OnlineExecHandler.h
--- a/Engine/Plugins/Online/OnlineServices/Source/OnlineServicesCommon/Public/Online/OnlineExecHandler.h
+++ b/Engine/Plugins/Online/OnlineServices/Source/OnlineServicesCommon/Public/Online/OnlineExecHandler.h
@@ -895,10 +895,12 @@ public:
 		Meta::VisitFields(Params, [&Cmd, &bSuccess, &Ar, &Services](const TCHAR* Name, auto& Field)
 		{
 			bool bResult = Private::ParseOnlineExecParams(Cmd, Field, Services);
+#if WITH_ENGINE
 			if (!bResult)
 			{
 				UE_LOG(LogConsoleResponse, Warning, TEXT("Failed to resolve outer field %s"), Name);
 			}
+#endif
 			bSuccess &= bResult;
 		});
 
@@ -907,10 +909,11 @@ public:
 			Ar.Log(TEXT("Failed to parse params"));
 			return false;
 		}
-
+#if WITH_ENGINE
 		if constexpr (Private::TOnlineInterfaceOperationMemberFunctionPtrTraits<MemberFunctionPtrType>::bAsync)
 		{
 			TOnlineAsyncOpHandle<OpType> AsyncOpHandle = Invoke(Function, Interface, MoveTemp(Params));
+
 			AsyncOpHandle.OnComplete([&Ar](const TOnlineResult<OpType>& Result)
 				{
 					UE_LOG(LogConsoleResponse, Display, TEXT("%s result: %s"), OpType::Name, *ToLogString(Result));
@@ -921,6 +924,7 @@ public:
 			TOnlineResult<OpType> Result = Invoke(Function, Interface, MoveTemp(Params));
 			UE_LOG(LogConsoleResponse, Display, TEXT("%s result: %s"), OpType::Name, *ToLogString(Result));
 		}
+#endif
 
 		return true;
 	}
