PR #9977
Fix generic service name computation in FOnlineServicesRegistry
(Contributed by [REDACTED_PERSON])

diff --git a/Engine/Plugins/Online/OnlineServices/Source/OnlineServicesInterface/Private/Online/OnlineServicesRegistry.cpp b/Engine/Plugins/Online/OnlineServices/Source/OnlineServicesInterface/Private/Online/OnlineServicesRegistry.cpp
--- a/Engine/Plugins/Online/OnlineServices/Source/OnlineServicesInterface/Private/Online/OnlineServicesRegistry.cpp
+++ b/Engine/Plugins/Online/OnlineServices/Source/OnlineServicesInterface/Private/Online/OnlineServicesRegistry.cpp
@@ -19,8 +19,41 @@ void FOnlineServicesRegistry::TearDown()
 	return TLazySingleton<FOnlineServicesRegistry>::TearDown();
 }
 
+EOnlineServices FOnlineServicesRegistry::ResolveServiceName(EOnlineServices OnlineServices) const
+{
+	if (OnlineServices == EOnlineServices::Default)
+	{
+		if (DefaultServiceOverride != EOnlineServices::Default)
+		{
+			OnlineServices = DefaultServiceOverride;
+		}
+		else
+		{
+			FString Value;
+
+			if (GConfig->GetString(TEXT("OnlineServices"), TEXT("DefaultServices"), Value, GEngineIni))
+			{
+				LexFromString(OnlineServices, *Value);
+			}
+		};
+	}
+	else if (OnlineServices == EOnlineServices::Platform)
+	{
+		FString Value;
+
+		if (GConfig->GetString(TEXT("OnlineServices"), TEXT("PlatformServices"), Value, GEngineIni))
+		{
+			LexFromString(OnlineServices, *Value);
+		}
+	}
+
+	return OnlineServices;
+}
+
 void FOnlineServicesRegistry::RegisterServicesFactory(EOnlineServices OnlineServices, TUniquePtr<IOnlineServicesFactory>&& Factory, int32 Priority)
 {
+	OnlineServices = ResolveServiceName(OnlineServices);
+
 	FFactoryAndPriority* ExistingFactoryAndPriority = ServicesFactories.Find(OnlineServices);
 	if (ExistingFactoryAndPriority == nullptr || ExistingFactoryAndPriority->Priority < Priority)
 	{
@@ -30,6 +63,8 @@ void FOnlineServicesRegistry::RegisterServicesFactory(EOnlineServices OnlineServ
 
 void FOnlineServicesRegistry::UnregisterServicesFactory(EOnlineServices OnlineServices, int32 Priority)
 {
+	OnlineServices = ResolveServiceName(OnlineServices);
+
 	FFactoryAndPriority* ExistingFactoryAndPriority = ServicesFactories.Find(OnlineServices);
 	if (ExistingFactoryAndPriority != nullptr && ExistingFactoryAndPriority->Priority == Priority)
 	{
@@ -41,6 +76,8 @@ void FOnlineServicesRegistry::UnregisterServicesFactory(EOnlineServices OnlineSe
 
 bool FOnlineServicesRegistry::IsLoaded(EOnlineServices OnlineServices, FName InstanceName) const
 {
+	OnlineServices = ResolveServiceName(OnlineServices);
+
 	bool bExists = false;
 	if (const TMap<FName, TSharedRef<IOnlineServices>>* OnlineServicesInstances = NamedServiceInstances.Find(OnlineServices))
 	{
@@ -51,28 +88,9 @@ bool FOnlineServicesRegistry::IsLoaded(EOnlineServices OnlineServices, FName Ins
 
 TSharedPtr<IOnlineServices> FOnlineServicesRegistry::GetNamedServicesInstance(EOnlineServices OnlineServices, FName InstanceName)
 {
-	TSharedPtr<IOnlineServices> Services;
-
-	if (OnlineServices == EOnlineServices::Default)
-	{
-		if(DefaultServiceOverride != EOnlineServices::Default)
-		{
-			OnlineServices = DefaultServiceOverride;
-		}
-		else
-		{
-			FString Value;
-			GConfig->GetString(TEXT("OnlineServices"), TEXT("DefaultServices"), Value, GEngineIni);
+	OnlineServices = ResolveServiceName(OnlineServices);
 
-			LexFromString(OnlineServices, *Value);
-		};
-	}
-	else if (OnlineServices == EOnlineServices::Platform)
-	{
-		FString Value;
-		GConfig->GetString(TEXT("OnlineServices"), TEXT("PlatformServices"), Value, GEngineIni);
-		LexFromString(OnlineServices, *Value);
-	}
+	TSharedPtr<IOnlineServices> Services;
 
 	if (OnlineServices < EOnlineServices::None)
 	{
@@ -97,6 +115,7 @@ TSharedPtr<IOnlineServices> FOnlineServicesRegistry::GetNamedServicesInstance(EO
 #if WITH_DEV_AUTOMATION_TESTS
 void FOnlineServicesRegistry::SetDefaultServiceOverride(EOnlineServices DefaultService)
 {
+	// No need to call ResolveServiceName here as a generic services name can be used as a Default Service Override
 	DefaultServiceOverride = DefaultService;
 }
 
@@ -108,6 +127,8 @@ void FOnlineServicesRegistry::ClearDefaultServiceOverride()
 
 void FOnlineServicesRegistry::DestroyNamedServicesInstance(EOnlineServices OnlineServices, FName InstanceName)
 {
+	OnlineServices = ResolveServiceName(OnlineServices);
+
 	if (TSharedRef<IOnlineServices>* ServicesPtr = NamedServiceInstances.FindOrAdd(OnlineServices).Find(InstanceName))
 	{
 		(*ServicesPtr)->Destroy();
@@ -118,6 +139,8 @@ void FOnlineServicesRegistry::DestroyNamedServicesInstance(EOnlineServices Onlin
 
 void FOnlineServicesRegistry::DestroyAllNamedServicesInstances(EOnlineServices OnlineServices)
 {
+	OnlineServices = ResolveServiceName(OnlineServices);
+
 	if (TMap<FName, TSharedRef<IOnlineServices>>* ServicesMapPtr = NamedServiceInstances.Find(OnlineServices))
 	{
 		for (const TPair<FName, TSharedRef<IOnlineServices>>& ServicesEntryRef : *ServicesMapPtr)
@@ -131,6 +154,8 @@ void FOnlineServicesRegistry::DestroyAllNamedServicesInstances(EOnlineServices O
 
 TSharedPtr<IOnlineServices> FOnlineServicesRegistry::CreateServices(EOnlineServices OnlineServices, FName InstanceName)
 {
+	OnlineServices = ResolveServiceName(OnlineServices);
+
 	TSharedPtr<IOnlineServices> Services;
 
 	FFactoryAndPriority* FactoryAndPriority = ServicesFactories.Find(OnlineServices);
diff --git a/Engine/Plugins/Online/OnlineServices/Source/OnlineServicesInterface/Public/Online/OnlineServicesRegistry.h b/Engine/Plugins/Online/OnlineServices/Source/OnlineServicesInterface/Public/Online/OnlineServicesRegistry.h
--- a/Engine/Plugins/Online/OnlineServices/Source/OnlineServicesInterface/Public/Online/OnlineServicesRegistry.h
+++ b/Engine/Plugins/Online/OnlineServices/Source/OnlineServicesInterface/Public/Online/OnlineServicesRegistry.h
@@ -145,6 +145,11 @@ private:
 	FOnlineServicesRegistry() {}
 	~FOnlineServicesRegistry();
 	friend FLazySingleton;
+
+	/**
+	* Resolves a generic (like Platform and Default) service enum value into the value of the real corresponding service.
+	*/
+	EOnlineServices ResolveServiceName(EOnlineServices OnlineServices) const;
 };
 
 /* UE::Online */ }
