Make sure to fill empty stats in null implementation when query for stats, to fix the warning in test: "Can't find stat Stat_Use_Largest when check if it can unlock achievement.";
Don't check achievement unlocking if the achievements states haven't been queried yet, and check achievement unlocking when achievement state got queried. It's to fix the warning in test: "Can't find state of achievement FAKE_ACHIEVEMENT when check if it's unlocked."

diff --git a/Engine/Plugins/Online/OnlineServices/Source/OnlineServicesCommon/Private/Online/AchievementsCommon.cpp b/Engine/Plugins/Online/OnlineServices/Source/OnlineServicesCommon/Private/Online/AchievementsCommon.cpp
--- a/Engine/Plugins/Online/OnlineServices/Source/OnlineServicesCommon/Private/Online/AchievementsCommon.cpp
+++ b/Engine/Plugins/Online/OnlineServices/Source/OnlineServicesCommon/Private/Online/AchievementsCommon.cpp
@@ -55,6 +55,28 @@ TOnlineAsyncOpHandle<FQueryAchievementDefinitions> FAchievementsCommon::QueryAch
 	return Operation->GetHandle();
 }
 
+void FAchievementsCommon::OnAchievementStatesQueried(const FAccountId& AccountId)
+{
+	if (!Config.bIsTitleManaged)
+	{
+		return;
+	}
+
+	TOnlineResult<FGetCachedStats> CachedStatsResult = Services.Get<FStatsCommon>()->GetCachedStats({});
+	check(CachedStatsResult.IsOk());
+	const TArray<FUserStats>& UsersStats = CachedStatsResult.GetOkValue().UsersStats;
+	for (const FUserStats& UserStats : UsersStats)
+	{
+		if (UserStats.AccountId == AccountId)
+		{
+			FStatsUpdated StatsUpdated;
+			StatsUpdated.LocalAccountId = AccountId;
+			StatsUpdated.UpdateUsersStats.Add(UserStats);
+			UnlockAchievementsByStats(StatsUpdated);
+		}
+	}
+}
+
 TOnlineResult<FGetAchievementIds> FAchievementsCommon::GetAchievementIds(FGetAchievementIds::Params&& Params)
 {
 	return TOnlineResult<FGetAchievementIds>(Errors::NotImplemented());
@@ -105,21 +127,24 @@ void FAchievementsCommon::UnlockAchievementsByStats(const FStatsUpdated& StatsUp
 	TArray<FAccountId> AccountIds;
 	for (const FUserStats& UserStats : StatsUpdated.UpdateUsersStats)
 	{
-		for (const TPair<FString, FStatValue>& StatPair : UserStats.Stats)
+		if (AchievementStates.Find(UserStats.AccountId))
 		{
-			for (const FAchievementUnlockRule& AchievementUnlockRule : Config.UnlockRules)
+			for (const TPair<FString, FStatValue>& StatPair : UserStats.Stats)
 			{
-				if (AchievementUnlockRule.ContainsStat(StatPair.Key))
+				for (const FAchievementUnlockRule& AchievementUnlockRule : Config.UnlockRules)
 				{
-					for (const FAchievementUnlockCondition& Condition : AchievementUnlockRule.Conditions)
+					if (AchievementUnlockRule.ContainsStat(StatPair.Key))
 					{
-						StatNames.AddUnique(Condition.StatName);
+						for (const FAchievementUnlockCondition& Condition : AchievementUnlockRule.Conditions)
+						{
+							StatNames.AddUnique(Condition.StatName);
+						}
 					}
 				}
 			}
-		}
 
-		AccountIds.AddUnique(UserStats.AccountId);
+			AccountIds.AddUnique(UserStats.AccountId);
+		}
 	}
 
 	if (StatNames.IsEmpty() || AccountIds.IsEmpty())
@@ -234,7 +259,7 @@ bool FAchievementsCommon::IsUnlocked(const FAccountId& AccountId, const FString&
 		return FMath::IsNearlyEqual(AchievementState.Progress, 1.0f);
 	}
 
-	UE_LOG(LogTemp, Warning, TEXT("Can't find state achievement %s when check if it's unlocked."), *AchievementName);
+	UE_LOG(LogTemp, Warning, TEXT("Can't find state of achievement %s when check if it's unlocked."), *AchievementName);
 
 	return false;
 }
diff --git a/Engine/Plugins/Online/OnlineServices/Source/OnlineServicesCommon/Public/Online/AchievementsCommon.h b/Engine/Plugins/Online/OnlineServices/Source/OnlineServicesCommon/Public/Online/AchievementsCommon.h
--- a/Engine/Plugins/Online/OnlineServices/Source/OnlineServicesCommon/Public/Online/AchievementsCommon.h
+++ b/Engine/Plugins/Online/OnlineServices/Source/OnlineServicesCommon/Public/Online/AchievementsCommon.h
@@ -76,6 +76,8 @@ public:
 protected:
 	TOnlineEventCallable<void(const FAchievementStateUpdated&)> OnAchievementStateUpdatedEvent;
 
+	void OnAchievementStatesQueried(const FAccountId& AccountId);
+
 	void UnlockAchievementsByStats(const FStatsUpdated& StatsUpdated);
 	void ExecuteUnlockRulesRelatedToStat(const FAccountId& AccountId, const FString& StatName, const TMap<FString, FStatValue>& Stats, TArray<FString>& OutAchievementsToUnlock);
 	bool MeetUnlockCondition(const FAchievementUnlockRule& AchievementUnlockRule, const TMap<FString, FStatValue>& Stats);
@@ -84,6 +86,9 @@ protected:
 	FOnlineEventDelegateHandle StatEventHandle;
 
 	FAchievementsCommonConfig Config;
+
+	using FAchievementStateMap = TMap<FString, FAchievementState>;
+	TMap<FAccountId, FAchievementStateMap> AchievementStates;
 };
 
 /* UE::Online */ }
diff --git a/Engine/Plugins/Online/OnlineServicesEOSGS/Source/Private/Online/AchievementsEOSGS.cpp b/Engine/Plugins/Online/OnlineServicesEOSGS/Source/Private/Online/AchievementsEOSGS.cpp
--- a/Engine/Plugins/Online/OnlineServicesEOSGS/Source/Private/Online/AchievementsEOSGS.cpp
+++ b/Engine/Plugins/Online/OnlineServicesEOSGS/Source/Private/Online/AchievementsEOSGS.cpp
@@ -301,6 +301,7 @@ TOnlineAsyncOpHandle<FQueryAchievementStates> FAchievementsEOSGS::QueryAchieveme
 			}
 			AchievementStates.Emplace(Params.LocalAccountId, MoveTemp(NewAchievementStates));
 			InAsyncOp.SetResult({});
+			OnAchievementStatesQueried(Params.LocalAccountId);
 		})
 		.Enqueue(GetSerialQueue());
 	}
diff --git a/Engine/Plugins/Online/OnlineServicesEOSGS/Source/Public/Online/AchievementsEOSGS.h b/Engine/Plugins/Online/OnlineServicesEOSGS/Source/Public/Online/AchievementsEOSGS.h
--- a/Engine/Plugins/Online/OnlineServicesEOSGS/Source/Public/Online/AchievementsEOSGS.h
+++ b/Engine/Plugins/Online/OnlineServicesEOSGS/Source/Public/Online/AchievementsEOSGS.h
@@ -38,10 +38,8 @@ protected:
 	EOS_HAchievements AchievementsHandle = nullptr;
 
 	using FAchievementDefinitionMap = TMap<FString, FAchievementDefinition>;
-	using FAchievementStateMap = TMap<FString, FAchievementState>;
 
 	TOptional<FAchievementDefinitionMap> AchievementDefinitions;
-	TMap<FAccountId, FAchievementStateMap> AchievementStates;
 
 	EOS_NotificationId NotifyAchievementsUnlockedNotificationId = EOS_INVALID_NOTIFICATIONID;
 };
diff --git a/Engine/Plugins/Online/OnlineServicesNull/Source/Private/Online/AchievementsNull.cpp b/Engine/Plugins/Online/OnlineServicesNull/Source/Private/Online/AchievementsNull.cpp
--- a/Engine/Plugins/Online/OnlineServicesNull/Source/Private/Online/AchievementsNull.cpp
+++ b/Engine/Plugins/Online/OnlineServicesNull/Source/Private/Online/AchievementsNull.cpp
@@ -106,6 +106,9 @@ TOnlineAsyncOpHandle<FQueryAchievementStates> FAchievementsNull::QueryAchievemen
 	}
 
 	Op->SetResult({});
+
+	OnAchievementStatesQueried(Op->GetParams().LocalAccountId);
+
 	return Op->GetHandle();
 }
 
diff --git a/Engine/Plugins/Online/OnlineServicesNull/Source/Private/Online/StatsNull.cpp b/Engine/Plugins/Online/OnlineServicesNull/Source/Private/Online/StatsNull.cpp
--- a/Engine/Plugins/Online/OnlineServicesNull/Source/Private/Online/StatsNull.cpp
+++ b/Engine/Plugins/Online/OnlineServicesNull/Source/Private/Online/StatsNull.cpp
@@ -94,6 +94,28 @@ TOnlineAsyncOpHandle<FUpdateStats> FStatsNull::UpdateStats(FUpdateStats::Params&
 	return Op->GetHandle();
 }
 
+namespace StatsNullPrivate
+{
+
+void ReadStatsFromCache(const FUserStats* ExistingUserStats, const TArray<FString>& StatNames, TMap<FString, FStatValue>& OutStats)
+{
+	for (const FString& StatName : StatNames)
+	{
+		if (ExistingUserStats)
+		{
+			if (const FStatValue* StatValue = ExistingUserStats->Stats.Find(StatName))
+			{
+				OutStats.Add(StatName, *StatValue);
+				continue;
+			}
+		}
+
+		OutStats.Add(StatName, FStatValue(static_cast<int64>(0)));
+	}
+}
+
+}
+
 TOnlineAsyncOpHandle<FQueryStats> FStatsNull::QueryStats(FQueryStats::Params&& Params)
 {
 	TOnlineAsyncOpRef<FQueryStats> Op = GetOp<FQueryStats>(MoveTemp(Params));
@@ -108,10 +130,8 @@ TOnlineAsyncOpHandle<FQueryStats> FStatsNull::QueryStats(FQueryStats::Params&& P
 	{
 		FQueryStats::Result Result;
 
-		if (FUserStats* ExistingUserStats = UsersStats.FindByPredicate(FFindUserStatsByAccountId(InAsyncOp.GetParams().TargetAccountId)))
-		{
-			Result.Stats = ExistingUserStats->Stats;
-		}
+		FUserStats* ExistingUserStats = UsersStats.FindByPredicate(FFindUserStatsByAccountId(InAsyncOp.GetParams().TargetAccountId));
+		StatsNullPrivate::ReadStatsFromCache(ExistingUserStats, InAsyncOp.GetParams().StatNames, Result.Stats);
 
 		InAsyncOp.SetResult(MoveTemp(Result));
 	})
@@ -135,10 +155,10 @@ TOnlineAsyncOpHandle<FBatchQueryStats> FStatsNull::BatchQueryStats(FBatchQuerySt
 
 		for (const FAccountId& TargetAccountId : InAsyncOp.GetParams().TargetAccountIds)
 		{
-			if (FUserStats* ExistingUserStats = UsersStats.FindByPredicate(FFindUserStatsByAccountId(TargetAccountId)))
-			{
-				Result.UsersStats.Emplace(*ExistingUserStats);
-			}
+			FUserStats& UserStats = Result.UsersStats.Emplace_GetRef();
+			UserStats.AccountId = TargetAccountId;
+			FUserStats* ExistingUserStats = UsersStats.FindByPredicate(FFindUserStatsByAccountId(TargetAccountId));
+			StatsNullPrivate::ReadStatsFromCache(ExistingUserStats, InAsyncOp.GetParams().StatNames, UserStats.Stats);
 		}
 
 		InAsyncOp.SetResult(MoveTemp(Result));
diff --git a/Engine/Plugins/Online/OnlineServicesNull/Source/Public/Online/AchievementsNull.h b/Engine/Plugins/Online/OnlineServicesNull/Source/Public/Online/AchievementsNull.h
--- a/Engine/Plugins/Online/OnlineServicesNull/Source/Public/Online/AchievementsNull.h
+++ b/Engine/Plugins/Online/OnlineServicesNull/Source/Public/Online/AchievementsNull.h
@@ -46,7 +46,6 @@ protected:
 	using FAchievementStateMap = TMap<FString, FAchievementState>;
 
 	bool bAchievementDefinitionsQueried = false;
-	TMap<FAccountId, FAchievementStateMap> AchievementStates;
 
 	FAchievementsNullConfig Config;
 
diff --git a/Engine/Plugins/Online/OnlineServicesOSSAdapter/Source/Private/Online/AchievementsOSSAdapter.cpp b/Engine/Plugins/Online/OnlineServicesOSSAdapter/Source/Private/Online/AchievementsOSSAdapter.cpp
--- a/Engine/Plugins/Online/OnlineServicesOSSAdapter/Source/Private/Online/AchievementsOSSAdapter.cpp
+++ b/Engine/Plugins/Online/OnlineServicesOSSAdapter/Source/Private/Online/AchievementsOSSAdapter.cpp
@@ -261,17 +261,18 @@ TOnlineAsyncOpHandle<FQueryAchievementStates> FAchievementsOSSAdapter::QueryAchi
 			return;
 		}
 
-		FAchievementStateMap& AchievementStates = UserToAchievementStates.Emplace(Op.GetParams().LocalAccountId);
+		FAchievementStateMap& LocalUserAchievementStates = AchievementStates.Emplace(Op.GetParams().LocalAccountId);
 
 		for (FOnlineAchievement& Achievement : Achievements)
 		{
-			FAchievementState& AchievementState = AchievementStates.Emplace(Achievement.Id);
+			FAchievementState& AchievementState = LocalUserAchievementStates.Emplace(Achievement.Id);
 			AchievementState.AchievementId = Achievement.Id;
 			AchievementState.Progress = Achievement.Progress / 100.f;
 			// AchievementState.UnlockTime left unset because this does not exist on per-user state, but on the common description...
 		}
 
 		Op.SetResult({});
+		OnAchievementStatesQueried(Op.GetParams().LocalAccountId);
 	})
 	.Enqueue(GetSerialQueue());
 
@@ -286,14 +287,14 @@ TOnlineResult<FGetAchievementState> FAchievementsOSSAdapter::GetAchievementState
 		return TOnlineResult<FGetAchievementState>(Errors::InvalidUser());
 	}
 
-	const FAchievementStateMap* AchievementStates = UserToAchievementStates.Find(Params.LocalAccountId);
-	if (!AchievementStates)
+	const FAchievementStateMap* LocalUserAchievementStates = AchievementStates.Find(Params.LocalAccountId);
+	if (!LocalUserAchievementStates)
 	{
 		// Call QueryAchievementStates first
 		return TOnlineResult<FGetAchievementState>(Errors::InvalidState());
 	}
 
-	const FAchievementState* AchievementState = AchievementStates->Find(Params.AchievementId);
+	const FAchievementState* AchievementState = LocalUserAchievementStates->Find(Params.AchievementId);
 	if (!AchievementState)
 	{
 		return TOnlineResult<FGetAchievementState>(Errors::NotFound());
diff --git a/Engine/Plugins/Online/OnlineServicesOSSAdapter/Source/Public/Online/AchievementsOSSAdapter.h b/Engine/Plugins/Online/OnlineServicesOSSAdapter/Source/Public/Online/AchievementsOSSAdapter.h
--- a/Engine/Plugins/Online/OnlineServicesOSSAdapter/Source/Public/Online/AchievementsOSSAdapter.h
+++ b/Engine/Plugins/Online/OnlineServicesOSSAdapter/Source/Public/Online/AchievementsOSSAdapter.h
@@ -39,9 +39,6 @@ protected:
 
 	using FAchievementDefinitionMap = TMap<FString, FAchievementDefinition>;
 	TOptional<FAchievementDefinitionMap> AchievementDefinitions;
-
-	using FAchievementStateMap = TMap<FString, FAchievementState>;
-	TMap<FAccountId, FAchievementStateMap> UserToAchievementStates;
 };
 
 /* UE::Online */ }
