Fixed mismatched copy constructor/assignment operator in FSessionInviteCommon.

diff --git a/Engine/Plugins/Online/OnlineServices/Source/OnlineServicesCommon/Public/Online/SessionsCommon.h b/Engine/Plugins/Online/OnlineServices/Source/OnlineServicesCommon/Public/Online/SessionsCommon.h
--- a/Engine/Plugins/Online/OnlineServices/Source/OnlineServicesCommon/Public/Online/SessionsCommon.h
+++ b/Engine/Plugins/Online/OnlineServices/Source/OnlineServicesCommon/Public/Online/SessionsCommon.h
@@ -175,7 +175,6 @@ class ONLINESERVICESCOMMON_API FSessionInviteCommon : public ISessionInvite
 {
 public:
 	FSessionInviteCommon() = default;
-	FSessionInviteCommon(const FSessionInviteCommon& Other) = default;
 	FSessionInviteCommon(const FAccountId& InRecipientId, const FAccountId& InSenderId, const FSessionInviteId& InInviteId, const FOnlineSessionId& InSessionId);
 
 	virtual const FAccountId GetRecipientId() const override { return RecipientId; }
