Fixed mismatched copy constructor/assignment operator in UE::Online::FSessionLAN.

diff --git a/Engine/Plugins/Online/OnlineServices/Source/OnlineServicesCommonEngineUtils/Public/Online/SessionsLAN.h b/Engine/Plugins/Online/OnlineServices/Source/OnlineServicesCommonEngineUtils/Public/Online/SessionsLAN.h
--- a/Engine/Plugins/Online/OnlineServices/Source/OnlineServicesCommonEngineUtils/Public/Online/SessionsLAN.h
+++ b/Engine/Plugins/Online/OnlineServices/Source/OnlineServicesCommonEngineUtils/Public/Online/SessionsLAN.h
@@ -27,7 +27,6 @@ class ONLINESERVICESCOMMONENGINEUTILS_API FSessionLAN : public FSessionCommon
 {
 public:
 	FSessionLAN();
-	FSessionLAN(const FSessionLAN& InSession) = default;
 
 	static FSessionLAN& Cast(FSessionCommon& InSession);
 	static const FSessionLAN& Cast(const ISession& InSession);
