Removed unsupported members of LobbyMember type.

diff --git a/Engine/Plugins/Online/OnlineServices/Source/OnlineServicesCommon/Private/Online/LobbiesCommonTypes.cpp b/Engine/Plugins/Online/OnlineServices/Source/OnlineServicesCommon/Private/Online/LobbiesCommonTypes.cpp
--- a/Engine/Plugins/Online/OnlineServices/Source/OnlineServicesCommon/Private/Online/LobbiesCommonTypes.cpp
+++ b/Engine/Plugins/Online/OnlineServices/Source/OnlineServicesCommon/Private/Online/LobbiesCommonTypes.cpp
@@ -89,8 +89,6 @@ TOnlineResult<FLobbyClientDataPrepareServiceSnapshot> FLobbyClientData::PrepareS
 				// New member, create an entry for them.
 				TSharedRef<FLobbyMemberInternal> NewMemberData = MakeShared<FLobbyMemberInternal>(SchemaRegistry);
 				NewMemberData->AccountId = MemberSnapshot.AccountId;
-				NewMemberData->PlatformAccountId = MemberSnapshot.PlatformAccountId;
-				NewMemberData->PlatformDisplayName = MemberSnapshot.PlatformDisplayName;
 				NewPreparedServiceChanges.NewMemberDataStorage.Add(NewMemberData);
 				PreparingMemberData = &NewMemberData.Get();
 			}
@@ -187,10 +185,6 @@ FLobbyClientDataCommitServiceSnapshot::Result FLobbyClientData::CommitServiceSna
 		// Commit attribute changes.
 		NewMember->SchemaData.CommitServiceSnapshot();
 		
-		// todo.
-		//NewMember->PlatformAccountId = ;
-		//NewMember->PlatformDisplayName = ;
-
 		// Add member to lobby structures.
 		AddMember(NewMember);
 
diff --git a/Engine/Plugins/Online/OnlineServices/Source/OnlineServicesCommon/Public/Online/LobbiesCommonTypes.h b/Engine/Plugins/Online/OnlineServices/Source/OnlineServicesCommon/Public/Online/LobbiesCommonTypes.h
--- a/Engine/Plugins/Online/OnlineServices/Source/OnlineServicesCommon/Public/Online/LobbiesCommonTypes.h
+++ b/Engine/Plugins/Online/OnlineServices/Source/OnlineServicesCommon/Public/Online/LobbiesCommonTypes.h
@@ -200,8 +200,6 @@ struct FLobbyServiceSnapshot
 struct FLobbyMemberServiceSnapshot
 {
 	FAccountId AccountId;
-	FAccountId PlatformAccountId;
-	FString PlatformDisplayName;
 	FSchemaServiceSnapshot SchemaServiceSnapshot;
 };
 
diff --git a/Engine/Plugins/Online/OnlineServices/Source/OnlineServicesInterface/Public/Online/Lobbies.h b/Engine/Plugins/Online/OnlineServices/Source/OnlineServicesInterface/Public/Online/Lobbies.h
--- a/Engine/Plugins/Online/OnlineServices/Source/OnlineServicesInterface/Public/Online/Lobbies.h
+++ b/Engine/Plugins/Online/OnlineServices/Source/OnlineServicesInterface/Public/Online/Lobbies.h
@@ -52,8 +52,6 @@ ONLINESERVICESINTERFACE_API void LexFromString(ELobbyMemberLeaveReason& OutLeave
 struct FLobbyMember
 {
 	FAccountId AccountId;
-	FAccountId PlatformAccountId;
-	FString PlatformDisplayName;
 	TMap<FSchemaAttributeId, FSchemaVariant> Attributes;
 	bool bIsLocalMember = false;
 };
@@ -863,8 +861,6 @@ namespace Meta {
 
 BEGIN_ONLINE_STRUCT_META(FLobbyMember)
 	ONLINE_STRUCT_FIELD(FLobbyMember, AccountId),
-	ONLINE_STRUCT_FIELD(FLobbyMember, PlatformAccountId),
-	ONLINE_STRUCT_FIELD(FLobbyMember, PlatformDisplayName),
 	ONLINE_STRUCT_FIELD(FLobbyMember, Attributes)
 END_ONLINE_STRUCT_META()
 
diff --git a/Engine/Plugins/Online/OnlineServicesEOSGS/Source/Private/Online/LobbiesEOSGSTypes.cpp b/Engine/Plugins/Online/OnlineServicesEOSGS/Source/Private/Online/LobbiesEOSGSTypes.cpp
--- a/Engine/Plugins/Online/OnlineServicesEOSGS/Source/Private/Online/LobbiesEOSGSTypes.cpp
+++ b/Engine/Plugins/Online/OnlineServicesEOSGS/Source/Private/Online/LobbiesEOSGSTypes.cpp
@@ -422,9 +422,6 @@ TDefaultErrorResultInternal<FLobbyMemberServiceSnapshot> FLobbyDetailsEOS::GetLo
 
 	FLobbyMemberServiceSnapshot LobbyMemberServcieSnapshot;
 	LobbyMemberServcieSnapshot.AccountId = MemberAccountId;
-	// Todo: 
-	//LobbyMemberServcieSnapshot.PlatformAccountId;
-	//LobbyMemberServcieSnapshot.PlatformDisplayName;
 
 	// Fetch attributes.
 	{
diff --git a/Engine/Plugins/Online/OnlineServicesNull/Source/Private/Online/LobbiesNull.cpp b/Engine/Plugins/Online/OnlineServicesNull/Source/Private/Online/LobbiesNull.cpp
--- a/Engine/Plugins/Online/OnlineServicesNull/Source/Private/Online/LobbiesNull.cpp
+++ b/Engine/Plugins/Online/OnlineServicesNull/Source/Private/Online/LobbiesNull.cpp
@@ -200,8 +200,6 @@ TSharedRef<FLobbyNull> FLobbiesNull::CreateNamedLobby(const FCreateLobby::Params
 	// construct the member data.
 	TSharedRef<FLobbyMember> NewMember = MakeShared<FLobbyMember>();
 	NewMember->AccountId = Params.LocalAccountId;
-	NewMember->PlatformAccountId = Params.LocalAccountId;
-	NewMember->PlatformDisplayName = TEXT("TEMP");
 	NewMember->Attributes = Params.UserAttributes;
 
 	Lobby->Data->Members.Add(Params.LocalAccountId, NewMember);
@@ -360,8 +358,6 @@ void FLobbiesNull::ReadLobbyFromPacket(FNboSerializeFromBuffer& Packet, const TS
 	// todo: make all the members (not just the host
 	TSharedRef<FLobbyMember> NewMember = MakeShared<FLobbyMember>();
 	NewMember->AccountId = OwningAccountId;
-	NewMember->PlatformAccountId = OwningAccountId;
-	NewMember->PlatformDisplayName = TEXT("TEMP");
 	Session->Data->Members.Add(OwningAccountId, NewMember);
 	AllLobbies.Add(Session->Data->LobbyId, Session);
 
