Added content resolver api and dummy implementation and content service integration.
Added support for structured logging of UE::Online::FOnlineError
Disabled by default.

diff --git a/Engine/Plugins/Online/OnlineServices/Source/OnlineServicesInterface/Private/Online/OnlineError.cpp b/Engine/Plugins/Online/OnlineServices/Source/OnlineServicesInterface/Private/Online/OnlineError.cpp
--- a/Engine/Plugins/Online/OnlineServices/Source/OnlineServicesInterface/Private/Online/OnlineError.cpp
+++ b/Engine/Plugins/Online/OnlineServices/Source/OnlineServicesInterface/Private/Online/OnlineError.cpp
@@ -1,6 +1,7 @@
 // Copyright Epic Games, Inc. All Rights Reserved.
 #include "Online/OnlineError.h"
 #include "Online/OnlineErrorCode.h"
+#include "Serialization/CompactBinaryWriter.h"
 
 namespace UE::Online{
 	namespace Errors {
@@ -61,5 +62,71 @@ namespace UE::Online{
 		return false;
 	}
 
+
+
+	void SerializeForLog(FCbWriter& Writer, const FOnlineError& OnlineError)
+	{
+		
+		Writer.BeginObject();
+		Writer.AddString(ANSITEXTVIEW("$type"), ANSITEXTVIEW("OnlineError"));
+
+
+		FText ErrorMessage;
+		FFormatNamedArguments ErrorMessageArgs;
+
+		bool bHasDetails = false;
+		bool bHasInner = false;
+
+		if (OnlineError == UE::Online::Errors::ErrorCode::Success)
+		{
+			FText SuccessText = NSLOCTEXT("OnlineError", "Success", "Success");
+			ErrorMessageArgs.Add(TEXT("ErrorCode"), SuccessText);
+			Writer.AddString(ANSITEXTVIEW("ErrorCode"), SuccessText.ToString());
+		}
+		else
+		{
+			FString ErrorCodeString = OnlineError.GetErrorId();
+			ErrorMessageArgs.Add(TEXT("ErrorCode"), FText::FromString(ErrorCodeString));
+			Writer.AddString(ANSITEXTVIEW("ErrorCode"), ErrorCodeString);
+		}
+
+		if (OnlineError.Details)
+		{
+			FText ErrorDetails = OnlineError.Details->GetText(OnlineError);
+			ErrorMessageArgs.Add(TEXT("ErrorDetails"), ErrorDetails);
+			Writer.AddString(ANSITEXTVIEW("ErrorDetails"), ErrorDetails.ToString());
+			bHasDetails = true;
+		}
+		
+		if (OnlineError.GetInner() != nullptr)
+		{
+			FText InnerText = OnlineError.GetInner()->GetText();
+			ErrorMessageArgs.Add(TEXT("InnerError"), InnerText);
+			Writer.AddString(ANSITEXTVIEW("InnerError"), InnerText.ToString());
+			bHasInner = true;
+		}
+
+		if (bHasInner && bHasDetails)
+		{
+			Writer.AddString("$text", FText::Format(NSLOCTEXT("OnlineError", "ErrorDetailsErrorInner", "{ErrorCode}, {ErrorDetails}, InnerError:{InnerError}"), ErrorMessageArgs).ToString());
+		}
+		else if (bHasDetails)
+		{
+			Writer.AddString("$text", FText::Format(NSLOCTEXT("OnlineError", "ErrorDetails", "{ErrorCode}, {ErrorDetails}"), ErrorMessageArgs).ToString());
+		}
+		else if (bHasInner)
+		{
+			Writer.AddString("$text", FText::Format(NSLOCTEXT("OnlineError", "ErrorInner", "{ErrorCode}, InnerError:{InnerError}"), ErrorMessageArgs).ToString());
+		}
+		else
+		{
+			Writer.AddString("$text", FText::Format(NSLOCTEXT("OnlineError", "ErrorCode", "{ErrorCode}"), ErrorMessageArgs).ToString());
+		}
+		
+
+		// Writer.AddString(ANSITEXTVIEW("errorNamespace"), OnlineError.ErrorNamespace);
+		Writer.EndObject();
+	}
+
 } /* namespace UE::Online */
 
diff --git a/Engine/Plugins/Online/OnlineServices/Source/OnlineServicesInterface/Public/Online/OnlineError.h b/Engine/Plugins/Online/OnlineServices/Source/OnlineServicesInterface/Public/Online/OnlineError.h
--- a/Engine/Plugins/Online/OnlineServices/Source/OnlineServicesInterface/Public/Online/OnlineError.h
+++ b/Engine/Plugins/Online/OnlineServices/Source/OnlineServicesInterface/Public/Online/OnlineError.h
@@ -166,10 +166,14 @@ private:
 	TSharedPtr<const IOnlineErrorDetails, ESPMode::ThreadSafe> Details;
 	TSharedPtr<const FOnlineError, ESPMode::ThreadSafe> Inner;
 	ErrorCodeType ErrorCode;
+
+	ONLINESERVICESINTERFACE_API friend void SerializeForLog(FCbWriter& Writer, const FOnlineError& OnlineError);
 };
 
 
 
+
+
 // the comparison with ErrorCodeType is the one I would expect to see in almost every case
 ONLINESERVICESINTERFACE_API bool operator==(const FOnlineError& Lhs, const FOnlineError& Rhs);
 inline bool operator!=(const FOnlineError& Lhs, const FOnlineError& Rhs) { return !(Lhs == Rhs); }
diff --git a/Engine/Plugins/Online/OnlineServicesOSSAdapter/Source/Private/Online/ErrorsOSSAdapter.cpp b/Engine/Plugins/Online/OnlineServicesOSSAdapter/Source/Private/Online/ErrorsOSSAdapter.cpp
--- a/Engine/Plugins/Online/OnlineServicesOSSAdapter/Source/Private/Online/ErrorsOSSAdapter.cpp
+++ b/Engine/Plugins/Online/OnlineServicesOSSAdapter/Source/Private/Online/ErrorsOSSAdapter.cpp
@@ -57,7 +57,7 @@ inline FOnlineError Internal_OssWrapInner(FOnlineError Error, const FOnlineError
 	return Internal_OssGetCommonError(Error, Result.GetErrorResult());
 }
 
-inline FOnlineError FromOssError(const FOnlineErrorOss& Result)
+ONLINESERVICESOSSADAPTER_API inline FOnlineError FromOssError(const FOnlineErrorOss& Result)
 {
 	FString ErrorCode;
 	FText ErrorMessage = Result.GetErrorMessage();
diff --git a/Engine/Plugins/Online/OnlineServicesOSSAdapter/Source/Public/Online/ErrorsOSSAdapter.h b/Engine/Plugins/Online/OnlineServicesOSSAdapter/Source/Public/Online/ErrorsOSSAdapter.h
--- a/Engine/Plugins/Online/OnlineServicesOSSAdapter/Source/Public/Online/ErrorsOSSAdapter.h
+++ b/Engine/Plugins/Online/OnlineServicesOSSAdapter/Source/Public/Online/ErrorsOSSAdapter.h
@@ -8,7 +8,7 @@ using FOnlineErrorOss = ::FOnlineError;
 namespace UE::Online::Errors {
 
 	UE_ONLINE_ERROR_CATEGORY(Oss, Engine, 0x7, "Oss")
-	FOnlineError FromOssError(const FOnlineErrorOss& Result);
+	ONLINESERVICESOSSADAPTER_API FOnlineError FromOssError(const FOnlineErrorOss& Result);
 	FOnlineError FromOssErrorCode(const FString& ErrorCode);
 	
 	// if you would like to add a parser that transforms custom platform errors into readable online service types, you can call this in your platform's error handler
