Fixed Online Tests WinGDK - error C2661: 'UE:FOnlineErrorDetails::FOnlineErrorDetails': no overloaded function takes 2 arguments in WindowsErrorsTest.cpp
 - Updates WindowsOnlineErrorDefinitions.h to pass in "Windows.<HRESULT>" code as FriendlyErrorCode

diff --git a/Engine/Plugins/Online/OnlineServices/Source/OnlineServicesInterface/Public/Online/Windows/WindowsOnlineErrorDefinitions.h b/Engine/Plugins/Online/OnlineServices/Source/OnlineServicesInterface/Public/Online/Windows/WindowsOnlineErrorDefinitions.h
--- a/Engine/Plugins/Online/OnlineServices/Source/OnlineServicesInterface/Public/Online/Windows/WindowsOnlineErrorDefinitions.h
+++ b/Engine/Plugins/Online/OnlineServices/Source/OnlineServicesInterface/Public/Online/Windows/WindowsOnlineErrorDefinitions.h
@@ -69,7 +69,7 @@ namespace UE::Online::Errors {
 
 		// This isn't localized, but Windows doesn't even have localized error codes anyway- do we need to add some method for localization? (e.g. try to look up the error first using some table generated by the WINDOWS_ERROR macro and use that error's FText)
 		// todo: can we get the error key somehow? (e.g. not "Operation aborted" which is ErrorString but rather E_ABORT). This is what the predefined errors use for the first parameter
-		return FOnlineError(ErrorCodeFromHRESULT(Result), MakeShared<FOnlineErrorDetails, ESPMode::ThreadSafe>(Internal_HResultToString(Result), FText::FromString(Internal_HResultToString(Result))));
+		return FOnlineError(ErrorCodeFromHRESULT(Result), MakeShared<FOnlineErrorDetails, ESPMode::ThreadSafe>(FString::Printf(TEXT("Windows.%x"), ErrorCodeFromHRESULT(Result)), Internal_HResultToString(Result), FText::FromString(Internal_HResultToString(Result))));
 	}
 
 
