Add UE_ONLINE_ERROR_EXTERNAL macro to enable online errors to be defined outside of the UE::Online::Errors namespace

diff --git a/Engine/Plugins/Online/OnlineServices/Source/OnlineServicesInterface/Public/Online/OnlineErrorCode.h b/Engine/Plugins/Online/OnlineServices/Source/OnlineServicesInterface/Public/Online/OnlineErrorCode.h
--- a/Engine/Plugins/Online/OnlineServices/Source/OnlineServicesInterface/Public/Online/OnlineErrorCode.h
+++ b/Engine/Plugins/Online/OnlineServices/Source/OnlineServicesInterface/Public/Online/OnlineErrorCode.h
@@ -48,7 +48,6 @@ ONLINESERVICESINTERFACE_API uint64 ErrorCodeSystem(ErrorCodeType ErrorCode);
 ONLINESERVICESINTERFACE_API uint64 ErrorCodeCategory(ErrorCodeType ErrorCode);
 ONLINESERVICESINTERFACE_API uint64 ErrorCodeValue(ErrorCodeType ErrorCode); 
 
-// needs to be inside Online::V2::Errors namespace
 // define an error category
 #define UE_ONLINE_ERROR_CATEGORY(Name, InSystem, Value, Description) \
 namespace ErrorCode { namespace Category { static constexpr int Name = Value; } }\
@@ -61,7 +60,7 @@ namespace ErrorCode { namespace Category { \
     static constexpr Name##_End = EndValue; \
 } }
 
-// Defines an individual error 
+// Defines an individual error. Not intended to be used directly.
 //	- Example access: (UE::Online::Errors::)NoConnection();
 	// - Can also specify an inner, e.g. RequestFailure(InvalidCreds());
 //  - Parameters:
@@ -79,20 +78,26 @@ namespace ErrorCode { namespace Category { \
 			{ \
 				InnerPtr = MakeShared<UE::Online::FOnlineError, ESPMode::ThreadSafe>(Inner.GetValue()); \
 			} \
-			return UE::Online::FOnlineError( ErrorCode::Create(UE::Online::Errors::ErrorCode::Category::CategoryName##_System, UE::Online::Errors::ErrorCode::Category::CategoryName, ErrorCodeValue), MakeShared<UE::Online::FOnlineErrorDetails, ESPMode::ThreadSafe>(FString(#CategoryName"."#Name), ErrorMessage, ErrorText), InnerPtr); \
+			return UE::Online::FOnlineError( UE::Online::Errors::ErrorCode::Create(ErrorCode::Category::CategoryName##_System, ErrorCode::Category::CategoryName, ErrorCodeValue), MakeShared<UE::Online::FOnlineErrorDetails, ESPMode::ThreadSafe>(FString(#CategoryName"."#Name), ErrorMessage, ErrorText), InnerPtr); \
 		}
 
+/** Macro to define an error within a category. Must be used within the UE::Online::Errors namespace. The error will be accessible as UE::Online::Errors::CategoryName::Name() */
 #define UE_ONLINE_ERROR(CategoryName, Name, ErrorCodeValue, ErrorMessage, ErrorText) \
     namespace ErrorCode { namespace CategoryName { static constexpr ErrorCodeType Name = Create(Category::CategoryName, ErrorCodeValue); } } \
 	namespace CategoryName { \
 		UE_ONLINE_ERROR_INTERNAL(CategoryName, Name, ErrorCodeValue, ErrorMessage, ErrorText)\
 	}
 
-// Same as above except no namespace around the error, e.g. Errors::NotImplemented instead of Errors::OnlineServices::NotImplemented
+/** Macro to define a common error in the UE::Online::Errors namespace. Must be used within the UE::Online::Errors namespace. The error will be accessible as UE::Online::Errors::Name() */
 #define UE_ONLINE_ERROR_COMMON(CategoryName, Name, ErrorCodeValue, ErrorMessage, ErrorText) \
     namespace ErrorCode { namespace CategoryName { static constexpr ErrorCodeType Name = Create(Category::CategoryName, ErrorCodeValue); } } \
 	UE_ONLINE_ERROR_INTERNAL(CategoryName, Name, ErrorCodeValue, ErrorMessage, ErrorText)
 
+/** Macro to define an error in a user defined namespace. The error will be accessible as CurrentNamespace::Name() */
+#define UE_ONLINE_ERROR_EXTERNAL(CategoryName, Name, ErrorCodeValue, ErrorMessage, ErrorText) \
+    namespace ErrorCode { namespace CategoryName { static constexpr UE::Online::Errors::ErrorCodeType Name = UE::Online::Errors::ErrorCode::Create(ErrorCode::Category::CategoryName, ErrorCodeValue); } } \
+	UE_ONLINE_ERROR_INTERNAL(CategoryName, Name, ErrorCodeValue, ErrorMessage, ErrorText)
+
 } /* namespace UE::Online::Errors */
 
 #if UE_ENABLE_INCLUDE_ORDER_DEPRECATED_IN_5_2
