- Add method to unregister component

diff --git a/Engine/Plugins/Online/OnlineServices/Source/OnlineServicesCommon/Public/Online/OnlineComponentRegistry.h b/Engine/Plugins/Online/OnlineServices/Source/OnlineServicesCommon/Public/Online/OnlineComponentRegistry.h
--- a/Engine/Plugins/Online/OnlineServices/Source/OnlineServicesCommon/Public/Online/OnlineComponentRegistry.h
+++ b/Engine/Plugins/Online/OnlineServices/Source/OnlineServicesCommon/Public/Online/OnlineComponentRegistry.h
@@ -27,6 +27,18 @@ public:
 		}
 	}
 
+	/**
+	 * Destroy and unregister a component of type ComponentType
+	 */
+	template <typename ComponentType>
+	void Unregister()
+	{
+		Components.RemoveAll([](const TUniquePtr<IComponentWrapper>& WrappedComponent)
+		{
+			return WrappedComponent->IsA(TOnlineTypeInfo<ComponentType>::GetTypeName());
+		});
+	}
+
 	/**
 	 * Get the component of type ComponentType
 	 * @return The component or null if one of ComponentType has not been registered
