Add GetInterface accessor to IOnlineServices to allow for additional/implementation specific interfaces to be accessed via the base IOnlineServices

diff --git a/Engine/Plugins/Online/OnlineServices/Source/OnlineServicesCommon/Private/Online/OnlineServicesCommon.cpp b/Engine/Plugins/Online/OnlineServices/Source/OnlineServicesCommon/Private/Online/OnlineServicesCommon.cpp
--- a/Engine/Plugins/Online/OnlineServices/Source/OnlineServicesCommon/Private/Online/OnlineServicesCommon.cpp
+++ b/Engine/Plugins/Online/OnlineServices/Source/OnlineServicesCommon/Private/Online/OnlineServicesCommon.cpp
@@ -129,6 +129,11 @@ FName FOnlineServicesCommon::GetInstanceName() const
 	return InstanceName;
 }
 
+void FOnlineServicesCommon::AssignBaseInterfaceSharedPtr(const FOnlineTypeName& TypeName, void* OutBaseInterfaceSP)
+{
+	Components.AssignBaseSharedPtr(TypeName, OutBaseInterfaceSP);
+}
+
 void FOnlineServicesCommon::RegisterComponents()
 {
 }
diff --git a/Engine/Plugins/Online/OnlineServices/Source/OnlineServicesCommon/Public/Online/OnlineComponentRegistry.h b/Engine/Plugins/Online/OnlineServices/Source/OnlineServicesCommon/Public/Online/OnlineComponentRegistry.h
--- a/Engine/Plugins/Online/OnlineServices/Source/OnlineServicesCommon/Public/Online/OnlineComponentRegistry.h
+++ b/Engine/Plugins/Online/OnlineServices/Source/OnlineServicesCommon/Public/Online/OnlineComponentRegistry.h
@@ -56,6 +56,16 @@ public:
 		return nullptr;
 	}
 
+	// Internal use only. Used to implement IOnlineServices::GetInterface()
+	void AssignBaseSharedPtr(const FOnlineTypeName& TypeName, void* OutBaseInterfaceSP)
+	{
+		const TUniquePtr<IComponentWrapper>* WrappedComponent = Components.FindByPredicate([&TypeName](const TUniquePtr<IComponentWrapper>& WrappedComponent) { return WrappedComponent->IsA(TypeName); });
+		if (WrappedComponent)
+		{
+			(*WrappedComponent)->AssignBaseSharedPtr(OutBaseInterfaceSP);
+		}
+	}
+
 	/**
 	 * Call a callable for each of the components
 	 */
@@ -74,6 +84,7 @@ private:
 	public:
 		virtual ~IComponentWrapper() {}
 		virtual void* GetBase() = 0;
+		virtual void AssignBaseSharedPtr(void* OutSP) = 0;
 		virtual bool IsA(FOnlineTypeName TypeName) const = 0;
 		virtual FOnlineTypeName GetTypeName() const = 0;
 		virtual IOnlineComponent* Get() = 0;
@@ -103,6 +114,13 @@ private:
 			return static_cast<void*>(static_cast<Meta::TBaseClass_T<T>*>(&Object));
 		}
 
+		virtual void AssignBaseSharedPtr(void* OutSP) override final
+		{
+			// assumes that component has AsShared that returns a shared pointer for the base class
+			TSharedPtr<Meta::TBaseClass_T<T>>& SP = *static_cast<TSharedPtr<Meta::TBaseClass_T<T>>*>(OutSP);
+			SP = Object.AsShared();
+		}
+
 		virtual bool IsA(FOnlineTypeName TypeName) const override final
 		{
 			return TOnlineTypeInfo<T>::IsA(TypeName);
diff --git a/Engine/Plugins/Online/OnlineServices/Source/OnlineServicesCommon/Public/Online/OnlineServicesCommon.h b/Engine/Plugins/Online/OnlineServices/Source/OnlineServicesCommon/Public/Online/OnlineServicesCommon.h
--- a/Engine/Plugins/Online/OnlineServices/Source/OnlineServicesCommon/Public/Online/OnlineServicesCommon.h
+++ b/Engine/Plugins/Online/OnlineServices/Source/OnlineServicesCommon/Public/Online/OnlineServicesCommon.h
@@ -49,6 +49,7 @@ public:
 	virtual IUserFilePtr GetUserFileInterface() override;
 	virtual TOnlineResult<FGetResolvedConnectString> GetResolvedConnectString(FGetResolvedConnectString::Params&& Params) override;
 	virtual FName GetInstanceName() const override;
+	virtual void AssignBaseInterfaceSharedPtr(const FOnlineTypeName& TypeName, void* OutBaseInterfaceSP) override final;
 
 	// FOnlineServicesCommon
 
diff --git a/Engine/Plugins/Online/OnlineServices/Source/OnlineServicesCommon/Public/Online/OnlineTypeInfo.h b/Engine/Plugins/Online/OnlineServices/Source/OnlineServicesCommon/Public/Online/OnlineTypeInfo.h
deleted file mode 100644
diff --git a/Engine/Plugins/Online/OnlineServices/Source/OnlineServicesInterface/Public/Online/OnlineServices.h b/Engine/Plugins/Online/OnlineServices/Source/OnlineServicesInterface/Public/Online/OnlineServices.h
--- a/Engine/Plugins/Online/OnlineServices/Source/OnlineServicesInterface/Public/Online/OnlineServices.h
+++ b/Engine/Plugins/Online/OnlineServices/Source/OnlineServicesInterface/Public/Online/OnlineServices.h
@@ -4,6 +4,7 @@
 
 #include "Online/CoreOnline.h"
 #include "Online/OnlineMeta.h"
+#include "Online/OnlineTypeInfo.h"
 
 namespace UE::Online { template <typename OpType> class TOnlineResult; }
 namespace UE::Online::Meta { template <typename StructType> struct TStructDetails; }
@@ -158,6 +159,29 @@ public:
 	 * @return instance name
 	 */
 	virtual FName GetInstanceName() const = 0;
+
+	/**
+	 * Retrieve any of the standard or implementation specific interfaces
+	 */
+	template <typename InterfaceType>
+	TSharedPtr<InterfaceType> GetInterface()
+	{
+		const FOnlineTypeName TypeName = TOnlineTypeInfo<InterfaceType>::GetTypeName();
+		using BaseInterfaceType = Meta::TBaseClass_T<InterfaceType>;
+		TSharedPtr<BaseInterfaceType> Interface;
+		AssignBaseInterfaceSharedPtr(TypeName, &Interface);
+		if constexpr (std::is_same_v<BaseInterfaceType, InterfaceType>)
+		{
+			return Interface;
+		}
+		else
+		{
+			return StaticCastSharedPtr<InterfaceType>(Interface);
+		}
+	}
+
+private:
+	virtual void AssignBaseInterfaceSharedPtr(const FOnlineTypeName& TypeName, void* InterfaceSP) = 0;
 };
 
 /**
diff --git a/Engine/Plugins/Online/OnlineServices/Source/OnlineServicesInterface/Public/Online/OnlineTypeInfo.h b/Engine/Plugins/Online/OnlineServices/Source/OnlineServicesInterface/Public/Online/OnlineTypeInfo.h
new file mode 100644
--- /dev/null
+++ b/Engine/Plugins/Online/OnlineServices/Source/OnlineServicesInterface/Public/Online/OnlineTypeInfo.h
@@ -0,0 +1,58 @@
+// Copyright Epic Games, Inc. All Rights Reserved.
+
+#pragma once
+
+#include "Misc/GeneratedTypeName.h"
+
+#include "Templates/Models.h"
+#include "Templates/TypeHash.h"
+
+namespace UE::Online::Meta { struct CSuperDefined; }
+
+namespace UE::Online {
+
+struct FOnlineTypeName
+{
+	const TCHAR* Name = nullptr;
+
+	bool operator==(const FOnlineTypeName& Other) const
+	{
+		return FCString::Strcmp(Name, Other.Name) == 0;
+	}
+
+	bool operator!=(const FOnlineTypeName& Other) const
+	{
+		return !operator==(Other);
+	}
+};
+
+inline uint32 GetTypeHash(const FOnlineTypeName& TypeName)
+{
+	return FCrc::Strihash_DEPRECATED(TypeName.Name);
+}
+
+template <typename T>
+struct TOnlineTypeInfo
+{
+	static inline FOnlineTypeName GetTypeName()
+	{
+		return FOnlineTypeName{ GetGeneratedTypeName<std::remove_reference_t<T>>() };
+	}
+
+	static inline bool IsA(FOnlineTypeName TypeName)
+	{
+		if (TypeName == GetTypeName())
+		{
+			return true;
+		}
+
+		if constexpr (TModels_V<Meta::CSuperDefined, T>)
+		{
+			return TOnlineTypeInfo<typename T::Super>::IsA(TypeName);
+		}
+
+		return false;
+	}
+};
+
+/* UE::Online */ }
