OnlineServices
- Handle more operation cache settings, duration based and provide a setting to cache error results
- Limit concurrent operations in flight
- Fix some operations not being removed/destroyed when we are done with them

diff --git a/Engine/Plugins/Online/OnlineServices/Source/OnlineServicesCommon/Private/Online/OnlineServicesCommon.cpp b/Engine/Plugins/Online/OnlineServices/Source/OnlineServicesCommon/Private/Online/OnlineServicesCommon.cpp
--- a/Engine/Plugins/Online/OnlineServices/Source/OnlineServicesCommon/Private/Online/OnlineServicesCommon.cpp
+++ b/Engine/Plugins/Online/OnlineServices/Source/OnlineServicesCommon/Private/Online/OnlineServicesCommon.cpp
@@ -10,6 +10,19 @@ DEFINE_LOG_CATEGORY(LogOnlineServices);
 
 namespace UE::Online {
 
+struct FOnlineServicesCommonConfig
+{
+	int32 MaxConcurrentOperations = 16;
+};
+
+namespace Meta {
+
+BEGIN_ONLINE_STRUCT_META(FOnlineServicesCommonConfig)
+	ONLINE_STRUCT_FIELD(FOnlineServicesCommonConfig, MaxConcurrentOperations)
+END_ONLINE_STRUCT_META()
+
+}
+
 uint32 FOnlineServicesCommon::NextInstanceIndex = 0;
 
 FOnlineServicesCommon::FOnlineServicesCommon(const FString& InConfigName, FName InInstanceName)
@@ -146,11 +159,15 @@ void FOnlineServicesCommon::Initialize()
 void FOnlineServicesCommon::PostInitialize()
 {
 	Components.Visit(&IOnlineComponent::PostInitialize);
+
+	LoadCommonConfig();
 }
 
 void FOnlineServicesCommon::UpdateConfig()
 {
 	Components.Visit(&IOnlineComponent::UpdateConfig);
+
+	LoadCommonConfig();
 }
 
 bool FOnlineServicesCommon::Tick(float DeltaSeconds)
@@ -226,4 +243,12 @@ bool FOnlineServicesCommon::Exec(UWorld* World, const TCHAR* Cmd, FOutputDevice&
 }
 #endif // UE_ALLOW_EXEC_COMMANDS
 
+void FOnlineServicesCommon::LoadCommonConfig()
+{
+	FOnlineServicesCommonConfig Config;
+	LoadConfig(Config);
+
+	ParallelQueue.SetMaxConcurrentOperations(Config.MaxConcurrentOperations);
+}
+
 /* UE::Online */ }
diff --git a/Engine/Plugins/Online/OnlineServices/Source/OnlineServicesCommon/Public/Online/OnlineAsyncOp.h b/Engine/Plugins/Online/OnlineServices/Source/OnlineServicesCommon/Public/Online/OnlineAsyncOp.h
--- a/Engine/Plugins/Online/OnlineServices/Source/OnlineServicesCommon/Public/Online/OnlineAsyncOp.h
+++ b/Engine/Plugins/Online/OnlineServices/Source/OnlineServicesCommon/Public/Online/OnlineAsyncOp.h
@@ -6,6 +6,7 @@
 #include "Online/OnlineResult.h"
 #include "Online/OnlineTypeInfo.h"
 #include "Async/Async.h"
+#include "Traits/ElementType.h"
 
 namespace UE::Online {
 
@@ -958,6 +959,7 @@ protected:
 	int NextStep = 0;
 
 	friend class FOnlineAsyncOpCache;
+	FOnlineEventDelegateHandle OpCacheHandle; // Delegate handle for FOnlineAsyncOpCache internal usage
 };
 
 template <typename OpType>
diff --git a/Engine/Plugins/Online/OnlineServices/Source/OnlineServicesCommon/Public/Online/OnlineAsyncOpCache.h b/Engine/Plugins/Online/OnlineServices/Source/OnlineServicesCommon/Public/Online/OnlineAsyncOpCache.h
--- a/Engine/Plugins/Online/OnlineServices/Source/OnlineServicesCommon/Public/Online/OnlineAsyncOpCache.h
+++ b/Engine/Plugins/Online/OnlineServices/Source/OnlineServicesCommon/Public/Online/OnlineAsyncOpCache.h
@@ -5,6 +5,7 @@
 #include "Online/CoreOnline.h"
 #include "Online/OnlineAsyncOp.h"
 #include "Online/OnlineMeta.h" // IWYU pragma: keep
+#include "Containers/Ticker.h"
 
 namespace UE::Online {
 
@@ -293,10 +294,11 @@ public:
 			LoadConfigFn(Config, ConfigSectionHeiarchy);
 			
 			Op = CreateOp<OpType, ParamsFuncsType>(MoveTemp(Params));
-			FOnlineEventDelegateHandle Handle = Op->OnComplete().Add(GetSharedThis(), [this](const TOnlineAsyncOp<OpType>& ThisOp, const TOnlineResult<OpType>&)
+			Op->OpCacheHandle = Op->OnComplete().Add(GetSharedThis(), [this, Config](const TOnlineAsyncOp<OpType>& ThisOp, const TOnlineResult<OpType>& Result)
 				{
-					// remove from cache if operation was canceled
-					if (ThisOp.GetState() == EAsyncOpState::Cancelled)
+					if ((ThisOp.GetState() == EAsyncOpState::Cancelled && !Config.bCacheError) ||
+						(Result.IsError() && !Config.bCacheError) ||
+						(Config.CacheExpiration == EOperationCacheExpirationPolicy::UponCompletion))
 					{
 						if constexpr (TModels_V<CLocalAccountIdDefined, typename OpType::Params>)
 						{
@@ -307,6 +309,26 @@ public:
 							Operations.Remove(FWrappedOperationKey::Create<OpType, ParamsFuncsType>(ThisOp.GetParams()));
 						}
 					}
+					else if (Config.CacheExpiration == EOperationCacheExpirationPolicy::Duration)
+					{
+						// Set timer with duration to remove cached operation
+						FTSTicker::GetCoreTicker().AddTicker(TEXT("OnlineAsyncOpCacheExpiry"), Config.CacheExpirySeconds,
+							[this, WeakOp = ThisOp.AsWeak()](float) {
+								TSharedPtr<const TOnlineAsyncOp<OpType>> PinnedOp = WeakOp.Pin();
+								if (PinnedOp.IsValid())
+								{
+									if constexpr (TModels_V<CLocalAccountIdDefined, typename OpType::Params>)
+									{
+										UserOperations.FindOrAdd(PinnedOp->GetParams().LocalAccountId).Remove(FWrappedOperationKey::Create<OpType, ParamsFuncsType>(PinnedOp->GetParams()));
+									}
+									else
+									{
+										Operations.Remove(FWrappedOperationKey::Create<OpType, ParamsFuncsType>(PinnedOp->GetParams()));
+									}
+								}
+ 								return false;
+							});
+					}
 				});
 		}
 
@@ -347,7 +369,7 @@ public:
 
 			Op = CreateOp<OpType, ParamsFuncsType>(MoveTemp(Params));
 			// remove from cache once operation has started. It is no longer mergeable at that point
-			FOnlineEventDelegateHandle Handle = Op->OnStart().Add(GetSharedThis(), [this](const TOnlineAsyncOp<OpType>& ThisOp)
+			Op->OpCacheHandle = Op->OnStart().Add(GetSharedThis(), [this](const TOnlineAsyncOp<OpType>& ThisOp)
 				{
 					if constexpr (TModels_V<CLocalAccountIdDefined, typename OpType::Params>)
 					{
@@ -420,7 +442,7 @@ private:
 				return true;
 			}
 
-			// other expiry conditions
+			// other expiry conditions are handled by removing the cached entry
 
 			return false;
 		}
diff --git a/Engine/Plugins/Online/OnlineServices/Source/OnlineServicesCommon/Public/Online/OnlineAsyncOpQueue.h b/Engine/Plugins/Online/OnlineServices/Source/OnlineServicesCommon/Public/Online/OnlineAsyncOpQueue.h
--- a/Engine/Plugins/Online/OnlineServices/Source/OnlineServicesCommon/Public/Online/OnlineAsyncOpQueue.h
+++ b/Engine/Plugins/Online/OnlineServices/Source/OnlineServicesCommon/Public/Online/OnlineAsyncOpQueue.h
@@ -109,16 +109,30 @@ public:
 
 	void Tick(float DeltaSeconds)
 	{
-		// TODO: Limit for number of in-flight operations
-		for (TOptional<TSharedRef<IWrappedOp>> Operation = QueuedOperations.Dequeue(); Operation.IsSet(); Operation = QueuedOperations.Dequeue())
+		while (InFlightOperations.Num() < MaxConcurrentOperations)
 		{
-			InFlightOperations.Add(Operation.GetValue());
+			TOptional<TSharedRef<IWrappedOp>> Operation = QueuedOperations.Dequeue();
+			if (!Operation.IsSet())
+			{
+				break;
+			}
+
+			InFlightOperations.Emplace(Operation.GetValue(), Operation.GetValue()->OnComplete().Add([this, WeakOp = TWeakPtr<IWrappedOp>(Operation.GetValue())]()
+				{
+					InFlightOperations.RemoveAll([WeakOp](const TPair<TSharedRef<IWrappedOp>, FOnlineEventDelegateHandle>& InFlightOperation)
+						{
+							return InFlightOperation.Key == WeakOp;
+						});
+				}));
 			Operation.GetValue()->Start();
 		}
 	}
 
+	void SetMaxConcurrentOperations(int InMaxConcurrentOperations) { check(InMaxConcurrentOperations > 0); MaxConcurrentOperations = InMaxConcurrentOperations; }
+
 protected:
-	TArray<TSharedRef<IWrappedOp>> InFlightOperations;
+	TArray<TPair<TSharedRef<IWrappedOp>, FOnlineEventDelegateHandle>> InFlightOperations;
+	int MaxConcurrentOperations = 16;
 };
 
 class FOnlineAsyncOpQueueSerial : public FOnlineAsyncOpQueue
diff --git a/Engine/Plugins/Online/OnlineServices/Source/OnlineServicesCommon/Public/Online/OnlineServicesCommon.h b/Engine/Plugins/Online/OnlineServices/Source/OnlineServicesCommon/Public/Online/OnlineServicesCommon.h
--- a/Engine/Plugins/Online/OnlineServices/Source/OnlineServicesCommon/Public/Online/OnlineServicesCommon.h
+++ b/Engine/Plugins/Online/OnlineServices/Source/OnlineServicesCommon/Public/Online/OnlineServicesCommon.h
@@ -285,19 +285,19 @@ public:
 
 	/* Get op (Interface) */
 	template <typename OpType>
-	TOnlineAsyncOpRef<OpType> GetOp(typename OpType::Params&& Params, const TArray<FString> ConfigSectionHeiarchy)
+	TOnlineAsyncOpRef<OpType> GetOp(typename OpType::Params&& Params, const TArray<FString>& ConfigSectionHeiarchy)
 	{
 		return OpCache.GetOp<OpType>(MoveTemp(Params), ConfigSectionHeiarchy);
 	}
 
-	template <typename OpType, typename ParamsFuncsType /*= TJoinableOpParamsFuncs<OpType>*/>
-	TOnlineAsyncOpRef<OpType> GetJoinableOp(typename OpType::Params&& Params, const TArray<FString> ConfigSectionHeiarchy)
+	template <typename OpType, typename ParamsFuncsType = TJoinableOpParamsFuncs<OpType>>
+	TOnlineAsyncOpRef<OpType> GetJoinableOp(typename OpType::Params&& Params, const TArray<FString>& ConfigSectionHeiarchy)
 	{
 		return OpCache.GetJoinableOp<OpType, ParamsFuncsType>(MoveTemp(Params), ConfigSectionHeiarchy);
 	}
 
-	template <typename OpType, typename ParamsFuncsType /*= TMergeableOpParamsFuncs<OpType>*/>
-	TOnlineAsyncOpRef<OpType> GetMergeableOp(typename OpType::Params&& Params, const TArray<FString> ConfigSectionHeiarchy)
+	template <typename OpType, typename ParamsFuncsType = TMergeableOpParamsFuncs<OpType>>
+	TOnlineAsyncOpRef<OpType> GetMergeableOp(typename OpType::Params&& Params, const TArray<FString>& ConfigSectionHeiarchy)
 	{
 		return OpCache.GetMergeableOp<OpType, ParamsFuncsType>(MoveTemp(Params), ConfigSectionHeiarchy);
 	}
@@ -320,6 +320,8 @@ public:
 	FOnlineAsyncOpCache OpCache;
 
 protected:
+	void LoadCommonConfig();
+
 	TMap<FString, TUniquePtr<IOnlineExecHandler>> ExecCommands;
 
 	static uint32 NextInstanceIndex;
