Fix hashing and comparison with joinable and mergeable ops

diff --git a/Engine/Plugins/Online/OnlineServices/Source/OnlineServicesCommon/Public/Online/OnlineAsyncOpCache.h b/Engine/Plugins/Online/OnlineServices/Source/OnlineServicesCommon/Public/Online/OnlineAsyncOpCache.h
--- a/Engine/Plugins/Online/OnlineServices/Source/OnlineServicesCommon/Public/Online/OnlineAsyncOpCache.h
+++ b/Engine/Plugins/Online/OnlineServices/Source/OnlineServicesCommon/Public/Online/OnlineAsyncOpCache.h
@@ -163,7 +163,7 @@ struct TJoinableOpParamsFuncs
 		bool bResult = true;
 		Meta::VisitFields<typename OpType::Params>([&bResult, &First, &Second](const auto& Field)
 		{
-			bResult = bResult && (First.*Field.Pointer) != (Second.*Field.Pointer);
+			bResult = bResult && (First.*Field.Pointer) == (Second.*Field.Pointer);
 		});
 		return bResult;
 	}
@@ -176,7 +176,7 @@ struct TJoinableOpParamsFuncs
 		Meta::VisitFields(Params,
 			[&CombinedHash](const TCHAR* FieldName, const auto& Field)
 			{
-				HashCombine(CombinedHash, GetTypeHash(Field));
+				CombinedHash = HashCombine(CombinedHash, GetTypeHash(Field));
 			});
 		return CombinedHash;
 	}
@@ -201,7 +201,7 @@ struct TMergeableOpParamsFuncs
 							return;
 						}
 					}
-					bResult = bResult && (First.*Field.Pointer) != (Second.*Field.Pointer);
+					bResult = bResult && (First.*Field.Pointer) == (Second.*Field.Pointer);
 				}
 			});
 		return bResult;
@@ -224,7 +224,7 @@ struct TMergeableOpParamsFuncs
 							return;
 						}
 					}
-					HashCombine(CombinedHash, GetTypeHash(Params.*Field.Pointer));
+					CombinedHash = HashCombine(CombinedHash, GetTypeHash(Params.*Field.Pointer));
 				}
 			});
 		return CombinedHash;
