OnlineServices
- Add ForEach and ForEachN continuation types that can be used to operate on elements of a TArray

diff --git a/Engine/Plugins/Online/OnlineServices/Source/OnlineServicesCommon/Private/Online/OnlineSyntaxTests.cpp b/Engine/Plugins/Online/OnlineServices/Source/OnlineServicesCommon/Private/Online/OnlineSyntaxTests.cpp
--- a/Engine/Plugins/Online/OnlineServices/Source/OnlineServicesCommon/Private/Online/OnlineSyntaxTests.cpp
+++ b/Engine/Plugins/Online/OnlineServices/Source/OnlineServicesCommon/Private/Online/OnlineSyntaxTests.cpp
@@ -186,6 +186,14 @@ inline void TestAsyncOpContinuationSyntax(FOnlineServicesCommon& Services)
 		.Then([](TOnlineAsyncOp<FTestOp>&, TPromise<int>&& Promise) { Promise.SetValue(1); })
 		.Then([](TOnlineAsyncOp<FTestOp>&, int, TPromise<FString>&& Promise) { Promise.SetValue(FString()); })
 		.Then([](TOnlineAsyncOp<FTestOp>&, const FString&, TPromise<int>&& Promise) { Promise.SetValue(1); })
+		.Then([](TOnlineAsyncOp<FTestOp>&, int) { return TArray<int>{ 1, 2, 3, 4 }; })
+		.ForEach([](TOnlineAsyncOp<FTestOp>&, int Value) { return Value; })
+		.ForEach([](TOnlineAsyncOp<FTestOp>&, int Value, TPromise<int>&& Promise) { Promise.SetValue(Value); })
+		.ForEach([](TOnlineAsyncOp<FTestOp>&, int Value) { return MakeFulfilledPromise<int>(Value).GetFuture(); })
+		.ForEachN(2, [](TOnlineAsyncOp<FTestOp>&, TArrayView<int> Value) { return TArray<int>(Value); })
+		.ForEachN(2, [](TOnlineAsyncOp<FTestOp>&, TArrayView<int> Value, TPromise<TArray<int>>&& Promise) { return Promise.SetValue(TArray<int>(Value)); })
+		.ForEachN(2, [](TOnlineAsyncOp<FTestOp>&, TArrayView<int> Value) { return MakeFulfilledPromise<TArray<int>>(Value).GetFuture(); })
+		.Then([](TOnlineAsyncOp<FTestOp>&, const TArray<int>&) { return 1; })
 		.Then([](TOnlineAsyncOp<FTestOp>&, int) {})
 		.Enqueue(Services.GetParallelQueue());
 	Op.GetHandle();
diff --git a/Engine/Plugins/Online/OnlineServices/Source/OnlineServicesCommon/Public/Online/OnlineAsyncOp.h b/Engine/Plugins/Online/OnlineServices/Source/OnlineServicesCommon/Public/Online/OnlineAsyncOp.h
--- a/Engine/Plugins/Online/OnlineServices/Source/OnlineServicesCommon/Public/Online/OnlineAsyncOp.h
+++ b/Engine/Plugins/Online/OnlineServices/Source/OnlineServicesCommon/Public/Online/OnlineAsyncOp.h
@@ -519,6 +519,248 @@ private:
 	TUniqueFunction<void()> ExecFunction;
 };
 
+template <typename ResultElementType>
+class TForEachStepBase : public IStep
+{
+public:
+	using ResultType = TArray<ResultElementType>;
+
+	TForEachStepBase(FOnlineAsyncExecutionPolicy&& InExecutionPolicy)
+		: ExecutionPolicy(MoveTemp(InExecutionPolicy))
+	{
+	}
+
+	~TForEachStepBase()
+	{
+		if (bResultConstructed)
+		{
+			DestructItem(Result.GetTypedPtr());
+		}
+	}
+
+	virtual const FOnlineAsyncExecutionPolicy& GetExecutionPolicy() const override
+	{
+		return ExecutionPolicy;
+	}
+
+	virtual void Execute() override
+	{
+		check(ExecFunction)
+		ExecFunction();
+	}
+
+	ResultType& GetResultRef()
+	{
+		return *Result.GetTypedPtr();
+	}
+
+protected:
+	FOnlineAsyncExecutionPolicy ExecutionPolicy;
+	TUniqueFunction<void()> ExecFunction;
+	TTypeCompatibleBytes<ResultType> Result;
+	std::atomic<int> NumResultsSet = 0;
+	bool bResultConstructed = false;
+};
+
+template <typename ResultElementType>
+class TForEachStep : public TForEachStepBase<ResultElementType>
+{
+public:
+	using typename TForEachStepBase<ResultElementType>::ResultType;
+
+	TForEachStep(FOnlineAsyncExecutionPolicy&& InExecutionPolicy)
+		: TForEachStepBase<ResultElementType>(MoveTemp(InExecutionPolicy))
+	{
+	}
+
+	template <typename OpType, typename LastResultType, typename CallableType>
+	void SetExecFunction(TOnlineAsyncOp<OpType>& InOperation, LastResultType&& InLastResult, CallableType&& InCallable)
+	{
+		this->ExecFunction = [this, WeakOperation = TWeakPtr<TOnlineAsyncOp<OpType>>(InOperation.AsShared()), &LastResult = InLastResult, Callable = MoveTemp(InCallable)]() mutable
+		{
+			TSharedPtr<TOnlineAsyncOp<OpType>> PinnedOperation = WeakOperation.Pin();
+			if (PinnedOperation)
+			{
+				int ResultIndex = 0;
+				new(this->Result.GetTypedPtr()) ResultType();
+				this->GetResultRef().SetNum(LastResult.Num());
+				this->bResultConstructed = true;
+
+				if constexpr (TOnlineAsyncOpCallableTraits<CallableType>::bRequiresPromise)
+				{
+					for (auto&& LastResultElement : LastResult)
+					{
+						TPromise<ResultElementType> Promise;
+						// set promise continuation before calling the callable so that we will complete the step as soon as the value is set
+						Promise.GetFuture()
+							.Next([this, WeakOperation, ThisResultIndex = ResultIndex++](const ResultElementType& Value)
+								{
+									TSharedPtr<TOnlineAsyncOp<OpType>> PinnedOperation2 = WeakOperation.Pin();
+									if (PinnedOperation2)
+									{
+										this->GetResultRef()[ThisResultIndex] = Value;
+										if (++this->NumResultsSet == this->GetResultRef().Num())
+										{
+											PinnedOperation2->ExecuteNextStep();
+										}
+									}
+								});
+
+						Callable(*PinnedOperation, MoveTempIfPossible(LastResultElement), MoveTemp(Promise));
+					}
+				}
+				else if constexpr (TOnlineAsyncOpCallableTraits<CallableType>::bAsyncResult)
+				{
+					for (auto&& LastResultElement : LastResult)
+					{
+						Callable(*PinnedOperation, MoveTempIfPossible(LastResultElement))
+							.Next([this, WeakOperation, ThisResultIndex = ResultIndex++](const ResultElementType& Value)
+								{
+									TSharedPtr<TOnlineAsyncOp<OpType>> PinnedOperation2 = WeakOperation.Pin();
+									if (PinnedOperation2)
+									{
+										this->GetResultRef()[ThisResultIndex] = Value;
+										if (++this->NumResultsSet == this->GetResultRef().Num())
+										{
+											PinnedOperation2->ExecuteNextStep();
+										}
+									}
+								});
+					}
+				}
+				else
+				{
+					for (auto&& LastResultElement : LastResult)
+					{
+						this->GetResultRef()[ResultIndex++] = Callable(*PinnedOperation, MoveTempIfPossible(LastResultElement));
+					}
+					PinnedOperation->ExecuteNextStep();
+				}
+			}
+		};
+	}
+};
+
+template <typename ResultElementType>
+class TForEachNStep : public TForEachStepBase<ResultElementType>
+{
+public:
+	using typename TForEachStepBase<ResultElementType>::ResultType;
+
+	TForEachNStep(int InBatchSize, FOnlineAsyncExecutionPolicy&& InExecutionPolicy)
+		: TForEachStepBase<ResultElementType>(MoveTemp(InExecutionPolicy))
+		, BatchSize(InBatchSize)
+	{
+	}
+
+	template <typename OpType, typename LastResultType, typename CallableType>
+	void SetExecFunction(TOnlineAsyncOp<OpType>& InOperation, LastResultType&& InLastResult, CallableType&& InCallable)
+	{
+		this->ExecFunction = [this, WeakOperation = TWeakPtr<TOnlineAsyncOp<OpType>>(InOperation.AsShared()), &LastResult = InLastResult, Callable = MoveTemp(InCallable)]() mutable
+		{
+			TSharedPtr<TOnlineAsyncOp<OpType>> PinnedOperation = WeakOperation.Pin();
+			if (PinnedOperation)
+			{
+				int ResultIndex = 0;
+				new(this->Result.GetTypedPtr()) ResultType();
+				this->GetResultRef().SetNumUninitialized(LastResult.Num());
+				this->bResultConstructed = true;
+
+				if constexpr (TOnlineAsyncOpCallableTraits<CallableType>::bRequiresPromise)
+				{
+					for (int Index = 0; Index < LastResult.Num(); Index += BatchSize)
+					{
+						int NumElementsInBatch = BatchSize;
+						if (Index + BatchSize > LastResult.Num())
+						{
+							NumElementsInBatch = LastResult.Num() - Index;
+						}
+
+						TPromise<TArray<ResultElementType>> Promise;
+						// set promise continuation before calling the callable so that we will complete the step as soon as the value is set
+						Promise.GetFuture()
+							.Next([this, WeakOperation, Index, NumElementsInBatch](const TArray<ResultElementType>& BatchResult)
+								{
+									TSharedPtr<TOnlineAsyncOp<OpType>> PinnedOperation2 = WeakOperation.Pin();
+									if (PinnedOperation2)
+									{
+										if (NumElementsInBatch != BatchResult.Num())
+										{
+											PinnedOperation2->SetError(Online::Errors::InvalidResults());
+											return;
+										}
+
+										MoveConstructItems(&this->GetResultRef()[Index], BatchResult.GetData(), BatchResult.Num());
+
+										if ((this->NumResultsSet += NumElementsInBatch) == this->GetResultRef().Num())
+										{
+											PinnedOperation2->ExecuteNextStep();
+										}
+									}
+								});
+
+						Callable(*PinnedOperation, MakeArrayView(&LastResult[Index], NumElementsInBatch), MoveTemp(Promise));
+					}
+				}
+				else if constexpr (TOnlineAsyncOpCallableTraits<CallableType>::bAsyncResult)
+				{
+					for (int Index = 0; Index < LastResult.Num(); Index += BatchSize)
+					{
+						int NumElementsInBatch = BatchSize;
+						if (Index + BatchSize > LastResult.Num())
+						{
+							NumElementsInBatch = LastResult.Num() - Index;
+						}
+
+						Callable(*PinnedOperation, MakeArrayView(&LastResult[Index], NumElementsInBatch))
+							.Next([this, WeakOperation, Index, NumElementsInBatch](const TArray<ResultElementType>& BatchResult)
+								{
+									TSharedPtr<TOnlineAsyncOp<OpType>> PinnedOperation2 = WeakOperation.Pin();
+									if (PinnedOperation2)
+									{
+										if (NumElementsInBatch != BatchResult.Num())
+										{
+											PinnedOperation2->SetError(Online::Errors::InvalidResults());
+											return;
+										}
+
+										MoveConstructItems(&this->GetResultRef()[Index], BatchResult.GetData(), BatchResult.Num());
+
+										if ((this->NumResultsSet += NumElementsInBatch) == this->GetResultRef().Num())
+										{
+											PinnedOperation2->ExecuteNextStep();
+										}
+									}
+								});
+					}
+				}
+				else
+				{
+					for (int Index = 0; Index < LastResult.Num(); Index += BatchSize)
+					{
+						int NumElementsInBatch = BatchSize;
+						if (Index + BatchSize > LastResult.Num())
+						{
+							NumElementsInBatch = LastResult.Num() - Index;
+						}
+						TArray<ResultElementType> BatchResult = Callable(*PinnedOperation, MakeArrayView(&LastResult[Index], NumElementsInBatch));
+						if (NumElementsInBatch != BatchResult.Num())
+						{
+							PinnedOperation->SetError(Online::Errors::InvalidResults());
+							break;
+						}
+						MoveConstructItems(&this->GetResultRef()[Index], BatchResult.GetData(), BatchResult.Num());
+					}
+					PinnedOperation->ExecuteNextStep();
+				}
+			}
+		};
+	}
+
+private:
+	int BatchSize;
+};
+
 // Provides Then continuation for both TOnlineAsyncOp and TOnlineChainableAsyncOp
 template <typename Outer, typename OpType, typename LastResultType>
 class TOnlineAsyncOpBase
@@ -542,6 +784,36 @@ public:
 	template <typename CallableType>
 	auto Then(CallableType&& Callable, FOnlineAsyncExecutionPolicy ExecutionPolicy = FOnlineAsyncExecutionPolicy::RunOnGameThread());
 
+	// Callable can take one of the following forms, where the second form is used when an asynchronous
+	//     call can set the promise with a value that is only valid for the duration of the Callable call.
+	//     This can be used when the previous continuation returned a TArray. The Callable will be called
+	//     once per element in the previous result. The results will be combined into a TArray
+	//   ResultType(AsyncOp, LastResultElement)
+	//     AsyncOp is a FOnlineAsyncOp& or TOnlineAsyncOp<OpType>&
+	//     LastResultElement is a LastResultElementType or const LastResultElementType&
+	//     ResultType is any type, or a TFuture to allow for an asynchronous result
+	//   void(AsyncOp, LastResultElement, TPromise<ResultType>&&)
+	//     AsyncOp is a FOnlineAsyncOp& or TOnlineAsyncOp<OpType>&
+	//     LastResultElement is a LastResultType or const LastResultType&
+	//     ResultType is any non-void type
+	template <typename CallableType> // previous value was TArray
+	auto ForEach(CallableType&& Callable, FOnlineAsyncExecutionPolicy ExecutionPolicy = FOnlineAsyncExecutionPolicy::RunOnGameThread());
+
+	// Callable can take one of the following forms, where the second form is used when an asynchronous
+	//     call can set the promise with a value that is only valid for the duration of the Callable call.
+	//     This can be used when the previous continuation returned a TArray. The Callable will be called
+	//     once per BatchSize elements in the previous result. The results will be combined into a TArray
+	//   ResultType(AsyncOp, LastResultElements)
+	//     AsyncOp is a FOnlineAsyncOp& or TOnlineAsyncOp<OpType>&
+	//     LastResultElements is a TArray<LastResultElementType> or const TArray<LastResultElementType>&
+	//     ResultType is any type, or a TFuture to allow for an asynchronous result
+	//   void(AsyncOp, LastResultElements, TPromise<ResultType>&&)
+	//     AsyncOp is a FOnlineAsyncOp& or TOnlineAsyncOp<OpType>&
+	//     LastResultElements is a TArray<LastResultElementType> or const TArray<LastResultElementType>&
+	//     ResultType is any non-void type
+	template <typename CallableType> // previous value was TArray
+	auto ForEachN(int BatchSize, CallableType&& Callable, FOnlineAsyncExecutionPolicy ExecutionPolicy = FOnlineAsyncExecutionPolicy::RunOnGameThread());
+
 protected:
 	LastResultType& LastResult;
 };
@@ -993,6 +1265,41 @@ auto TOnlineAsyncOpBase<Outer, OpType, LastResultType>::Then(CallableType&& InCa
 	}
 }
 
+template <typename Outer, typename OpType, typename LastResultType>
+template <typename CallableType>
+auto TOnlineAsyncOpBase<Outer, OpType, LastResultType>::ForEach(CallableType&& InCallable, FOnlineAsyncExecutionPolicy ExecutionPolicy)
+{
+	using ResultType = typename TOnlineAsyncOpCallableTraits<CallableType>::ResultType;
+
+	TOnlineAsyncOp<OpType>& Op = static_cast<Outer*>(this)->GetOwningOperation();
+
+	TForEachStep<ResultType>* Step = new TForEachStep<ResultType>(MoveTemp(ExecutionPolicy));
+	TUniquePtr<IStep> StepPtr(Step);
+	Step->SetExecFunction(Op, LastResult, MoveTemp(InCallable));
+
+	Op.AddStep(MoveTemp(StepPtr));
+
+	return TOnlineChainableAsyncOp<OpType, TArray<ResultType>>(Op, Step->GetResultRef());
+}
+
+template <typename Outer, typename OpType, typename LastResultType>
+template <typename CallableType>
+auto TOnlineAsyncOpBase<Outer, OpType, LastResultType>::ForEachN(int BatchSize, CallableType&& InCallable, FOnlineAsyncExecutionPolicy ExecutionPolicy)
+{
+	using ResultTypeElement = TElementType_T<typename TOnlineAsyncOpCallableTraits<CallableType>::ResultType>;
+
+	TOnlineAsyncOp<OpType>& Op = static_cast<Outer*>(this)->GetOwningOperation();
+
+	TForEachNStep<ResultTypeElement>* Step = new TForEachNStep<ResultTypeElement>(BatchSize, MoveTemp(ExecutionPolicy));
+	TUniquePtr<IStep> StepPtr(Step);
+	Step->SetExecFunction(Op, LastResult, MoveTemp(InCallable));
+
+	Op.AddStep(MoveTemp(StepPtr));
+
+	return TOnlineChainableAsyncOp<OpType, TArray<ResultTypeElement>>(Op, Step->GetResultRef());
+}
+
+
 template <typename Outer, typename OpType>
 template <typename CallableType>
 auto TOnlineAsyncOpBase<Outer, OpType, void>::Then(CallableType&& InCallable, FOnlineAsyncExecutionPolicy ExecutionPolicy)
