Add FAccountInfoRegistry::FindNoLock method

diff --git a/Engine/Plugins/Online/OnlineServices/Source/OnlineServicesCommon/Private/Online/AuthCommon.cpp b/Engine/Plugins/Online/OnlineServices/Source/OnlineServicesCommon/Private/Online/AuthCommon.cpp
--- a/Engine/Plugins/Online/OnlineServices/Source/OnlineServicesCommon/Private/Online/AuthCommon.cpp
+++ b/Engine/Plugins/Online/OnlineServices/Source/OnlineServicesCommon/Private/Online/AuthCommon.cpp
@@ -11,15 +11,13 @@ namespace UE::Online {
 TSharedPtr<FAccountInfo> FAccountInfoRegistry::Find(FPlatformUserId PlatformUserId) const
 {
 	FReadScopeLock Lock(IndexLock);
-	const TSharedRef<FAccountInfo>* FoundPtr = AuthDataByPlatformUserId.Find(PlatformUserId);
-	return FoundPtr ? TSharedPtr<FAccountInfo>(*FoundPtr) : TSharedPtr<FAccountInfo>();
+	return FindNoLock(PlatformUserId);
 }
 
 TSharedPtr<FAccountInfo> FAccountInfoRegistry::Find(FAccountId AccountId) const
 {
 	FReadScopeLock Lock(IndexLock);
-	const TSharedRef<FAccountInfo>* FoundPtr = AuthDataByOnlineAccountIdHandle.Find(AccountId);
-	return FoundPtr ? TSharedPtr<FAccountInfo>(*FoundPtr) : TSharedPtr<FAccountInfo>();
+	return FindNoLock(AccountId);
 }
 
 TArray<TSharedRef<FAccountInfo>> FAccountInfoRegistry::GetAllAccountInfo(TFunction<bool(const TSharedRef<FAccountInfo>&)> Predicate) const
@@ -41,6 +39,18 @@ TArray<TSharedRef<FAccountInfo>> FAccountInfoRegistry::GetAllAccountInfo(TFuncti
 	return Result;
 }
 
+TSharedPtr<FAccountInfo> FAccountInfoRegistry::FindNoLock(FPlatformUserId PlatformUserId) const
+{
+	const TSharedRef<FAccountInfo>* FoundPtr = AuthDataByPlatformUserId.Find(PlatformUserId);
+	return FoundPtr ? TSharedPtr<FAccountInfo>(*FoundPtr) : TSharedPtr<FAccountInfo>();
+}
+
+TSharedPtr<FAccountInfo> FAccountInfoRegistry::FindNoLock(FAccountId AccountId) const
+{
+	const TSharedRef<FAccountInfo>* FoundPtr = AuthDataByOnlineAccountIdHandle.Find(AccountId);
+	return FoundPtr ? TSharedPtr<FAccountInfo>(*FoundPtr) : TSharedPtr<FAccountInfo>();
+}
+
 void FAccountInfoRegistry::DoRegister(const TSharedRef<FAccountInfo>& AccountInfo)
 {
 	if (ensure(AccountInfo->PlatformUserId != PLATFORMUSERID_NONE))
diff --git a/Engine/Plugins/Online/OnlineServices/Source/OnlineServicesCommon/Public/Online/AuthCommon.h b/Engine/Plugins/Online/OnlineServices/Source/OnlineServicesCommon/Public/Online/AuthCommon.h
--- a/Engine/Plugins/Online/OnlineServices/Source/OnlineServicesCommon/Public/Online/AuthCommon.h
+++ b/Engine/Plugins/Online/OnlineServices/Source/OnlineServicesCommon/Public/Online/AuthCommon.h
@@ -19,6 +19,9 @@ public:
 	TArray<TSharedRef<FAccountInfo>> GetAllAccountInfo(TFunction<bool(const TSharedRef<FAccountInfo>&)> Predicate) const;
 
 protected:
+	TSharedPtr<FAccountInfo> FindNoLock(FPlatformUserId PlatformUserId) const;
+	TSharedPtr<FAccountInfo> FindNoLock(FAccountId AccountId) const;
+
 	virtual void DoRegister(const TSharedRef<FAccountInfo>& AccountInfo);
 	virtual void DoUnregister(const TSharedRef<FAccountInfo>& AccountInfo);
 
