Fix typo

diff --git a/Engine/Plugins/Online/OnlineServices/Source/OnlineServicesCommon/Public/Online/OnlineExecHandler.h b/Engine/Plugins/Online/OnlineServices/Source/OnlineServicesCommon/Public/Online/OnlineExecHandler.h
--- a/Engine/Plugins/Online/OnlineServices/Source/OnlineServicesCommon/Public/Online/OnlineExecHandler.h
+++ b/Engine/Plugins/Online/OnlineServices/Source/OnlineServicesCommon/Public/Online/OnlineExecHandler.h
@@ -52,7 +52,7 @@ template <typename OptionalType> struct TOptional;
 	Full command examples:
 		OnlineServices Index=0 Auth Login Null username s:pwd []
 		OnlineServices Index=0 Presence UpdatePresence 0 {0 Unknown Unknown Unknown Test1 Test2 {a=b}}
-		OnlineServices Index=0 Pressence QueryPresence 0 0 true
+		OnlineServices Index=0 Presence QueryPresence 0 0 true
 */
 
 class UWorld;
