Replacing more unnecessary warnings with Verbose

diff --git a/Engine/Plugins/Online/OnlineServices/Source/OnlineServicesCommonEngineUtils/Private/Online/SessionsLAN.cpp b/Engine/Plugins/Online/OnlineServices/Source/OnlineServicesCommonEngineUtils/Private/Online/SessionsLAN.cpp
--- a/Engine/Plugins/Online/OnlineServices/Source/OnlineServicesCommonEngineUtils/Private/Online/SessionsLAN.cpp
+++ b/Engine/Plugins/Online/OnlineServices/Source/OnlineServicesCommonEngineUtils/Private/Online/SessionsLAN.cpp
@@ -347,7 +347,7 @@ void FSessionsLAN::FindLANSessions(const FAccountId& LocalAccountId)
 	}
 	else
 	{
-		UE_LOG(LogOnlineServices, Warning, TEXT("[]FSessionsLAN::FindLANSessions] Searching...."));
+		UE_LOG(LogOnlineServices, Verbose, TEXT("[]FSessionsLAN::FindLANSessions] Searching...."));
 	}
 }
 
