OnlineError logging changes-
Fixed ToLogString discarding the prefix even if the option is set to true. The new format looks something like this-
   [[REDACTED_PHONE]] NotImplemented (EOS_MissingInterface)
ToLogString now logs the human readable ID instead of the binary ID (i.e. [[REDACTED_PHONE]] instead of the binary representation of that)

diff --git a/Engine/Plugins/Online/OnlineServices/Source/OnlineServicesInterface/Public/Online/OnlineError.h b/Engine/Plugins/Online/OnlineServices/Source/OnlineServicesInterface/Public/Online/OnlineError.h
--- a/Engine/Plugins/Online/OnlineServices/Source/OnlineServicesInterface/Public/Online/OnlineError.h
+++ b/Engine/Plugins/Online/OnlineServices/Source/OnlineServicesInterface/Public/Online/OnlineError.h
@@ -94,11 +94,11 @@ public:
 			FString LogPrefix = TEXT("");
 			if(bIncludePrefix)
 			{
-				LogPrefix = FString::Printf(TEXT("Error [%s]: "), *GetErrorId());
+				LogPrefix = FString::Printf(TEXT("[%s] "), *GetErrorId());
 			}
 
 			FString LogString = Details->GetLogString(*this);
-			MyLogString = Details->GetLogString(*this);
+			MyLogString = FString::Printf(TEXT("%s%s"), *LogPrefix, *Details->GetLogString(*this));
 		}
 		else if (ErrorCode == Errors::ErrorCode::Success)
 		{
@@ -109,7 +109,7 @@ public:
 		}
 		else
 		{
-			MyLogString = FString::Printf(TEXT("%x"), ErrorCode); // would use %s here and below if we were to use a string ErrorCodeType
+			MyLogString = FString::Printf(TEXT("[%s]"), *GetErrorId()); 
 		}
 
 		if (GetInner() != nullptr)
