Several small fixes to OSS Adapter:
 - Properly initializing achievements
 - Avoid crashing on missing interface
 - Improve check so it is not triggered on valid non happy path

diff --git a/Engine/Plugins/Online/OnlineServices/Source/OnlineServicesCommon/Private/Online/AchievementsCommon.cpp b/Engine/Plugins/Online/OnlineServices/Source/OnlineServicesCommon/Private/Online/AchievementsCommon.cpp
--- a/Engine/Plugins/Online/OnlineServices/Source/OnlineServicesCommon/Private/Online/AchievementsCommon.cpp
+++ b/Engine/Plugins/Online/OnlineServices/Source/OnlineServicesCommon/Private/Online/AchievementsCommon.cpp
@@ -21,7 +21,10 @@ void FAchievementsCommon::Initialize()
 {
 	TOnlineComponent<IAchievements>::Initialize();
 
-	StatEventHandle = Services.Get<IStats>()->OnStatsUpdated().Add([this](const FStatsUpdated& StatsUpdated) { UnlockAchievementsByStats(StatsUpdated); });
+	if (IStats* StatsInterface = Services.Get<IStats>())
+	{
+		StatEventHandle = StatsInterface->OnStatsUpdated().Add([this](const FStatsUpdated& StatsUpdated) { UnlockAchievementsByStats(StatsUpdated); });
+	}
 }
 
 void FAchievementsCommon::Shutdown()
@@ -62,17 +65,20 @@ void FAchievementsCommon::OnAchievementStatesQueried(const FAccountId& AccountId
 		return;
 	}
 
-	TOnlineResult<FGetCachedStats> CachedStatsResult = Services.Get<FStatsCommon>()->GetCachedStats({});
-	check(CachedStatsResult.IsOk());
-	const TArray<FUserStats>& UsersStats = CachedStatsResult.GetOkValue().UsersStats;
-	for (const FUserStats& UserStats : UsersStats)
+	if (IStats* Stats = Services.Get<IStats>())
 	{
-		if (UserStats.AccountId == AccountId)
+		TOnlineResult<FGetCachedStats> CachedStatsResult = Stats->GetCachedStats({});
+		check(CachedStatsResult.IsOk());
+		const TArray<FUserStats>& UsersStats = CachedStatsResult.GetOkValue().UsersStats;
+		for (const FUserStats& UserStats : UsersStats)
 		{
-			FStatsUpdated StatsUpdated;
-			StatsUpdated.LocalAccountId = AccountId;
-			StatsUpdated.UpdateUsersStats.Add(UserStats);
-			UnlockAchievementsByStats(StatsUpdated);
+			if (UserStats.AccountId == AccountId)
+			{
+				FStatsUpdated StatsUpdated;
+				StatsUpdated.LocalAccountId = AccountId;
+				StatsUpdated.UpdateUsersStats.Add(UserStats);
+				UnlockAchievementsByStats(StatsUpdated);
+			}
 		}
 	}
 }
@@ -123,6 +129,12 @@ void FAchievementsCommon::UnlockAchievementsByStats(const FStatsUpdated& StatsUp
 		return;
 	}
 
+	IStats* StatsInterface = Services.Get<IStats>();
+	if (!StatsInterface)
+	{
+		return;
+	}
+
 	TArray<FString> StatNames;
 	TArray<FAccountId> AccountIds;
 	for (const FUserStats& UserStats : StatsUpdated.UpdateUsersStats)
@@ -157,7 +169,7 @@ void FAchievementsCommon::UnlockAchievementsByStats(const FStatsUpdated& StatsUp
 	BatchQueryStatsParam.TargetAccountIds = MoveTemp(AccountIds);
 	BatchQueryStatsParam.StatNames = MoveTemp(StatNames);
 
-	Services.Get<FStatsCommon>()->BatchQueryStats(MoveTemp(BatchQueryStatsParam))
+	StatsInterface->BatchQueryStats(MoveTemp(BatchQueryStatsParam))
 	.OnComplete([this](const TOnlineResult<FBatchQueryStats>& Result)
 	{
 		if (Result.IsOk())
@@ -208,9 +220,15 @@ bool FAchievementsCommon::MeetUnlockCondition(const FAchievementUnlockRule& Achi
 		return false;
 	}
 
+	FStatsCommon* StatsInterface = Services.Get<FStatsCommon>();
+	if (!StatsInterface)
+	{
+		return false;
+	}
+
 	for (const FAchievementUnlockCondition& Condition : AchievementUnlockRule.Conditions)
 	{
-		if (const FStatDefinition* StatDefinition = Services.Get<FStatsCommon>()->GetStatDefinition(Condition.StatName))
+		if (const FStatDefinition* StatDefinition = StatsInterface->GetStatDefinition(Condition.StatName))
 		{
 			const FStatValue* StatValue = Stats.Find(Condition.StatName);
 			if (!StatValue)
diff --git a/Engine/Plugins/Online/OnlineServices/Source/OnlineServicesCommon/Private/Online/LeaderboardsCommon.cpp b/Engine/Plugins/Online/OnlineServices/Source/OnlineServicesCommon/Private/Online/LeaderboardsCommon.cpp
--- a/Engine/Plugins/Online/OnlineServices/Source/OnlineServicesCommon/Private/Online/LeaderboardsCommon.cpp
+++ b/Engine/Plugins/Online/OnlineServices/Source/OnlineServicesCommon/Private/Online/LeaderboardsCommon.cpp
@@ -73,7 +73,10 @@ void FLeaderboardsCommon::Initialize()
 {
 	Super::Initialize();
 
-	StatEventHandle = Services.Get<IStats>()->OnStatsUpdated().Add([this](const FStatsUpdated& StatsUpdated) { WriteLeaderboardsByStats(StatsUpdated); });
+	if (IStats* StatsInterface = Services.Get<IStats>())
+	{
+		StatEventHandle = StatsInterface->OnStatsUpdated().Add([this](const FStatsUpdated& StatsUpdated) { WriteLeaderboardsByStats(StatsUpdated); });
+	}
 }
 
 void FLeaderboardsCommon::UpdateConfig()
diff --git a/Engine/Plugins/Online/OnlineServicesOSSAdapter/Source/Private/Online/AchievementsOSSAdapter.cpp b/Engine/Plugins/Online/OnlineServicesOSSAdapter/Source/Private/Online/AchievementsOSSAdapter.cpp
--- a/Engine/Plugins/Online/OnlineServicesOSSAdapter/Source/Private/Online/AchievementsOSSAdapter.cpp
+++ b/Engine/Plugins/Online/OnlineServicesOSSAdapter/Source/Private/Online/AchievementsOSSAdapter.cpp
@@ -156,6 +156,7 @@ TOnlineAsyncOpHandle<FQueryAchievementDefinitions> FAchievementsOSSAdapter::Quer
 			Definitions.Emplace(MoveTemp(Achievement.Id), MoveTemp(Definition));
 		}
 
+		AchievementDefinitions = MoveTemp(Definitions);
 		Op.SetResult({});
 	})
 	.Enqueue(GetSerialQueue());
diff --git a/Engine/Plugins/Online/OnlineServicesOSSAdapter/Source/Private/Online/AuthOSSAdapter.cpp b/Engine/Plugins/Online/OnlineServicesOSSAdapter/Source/Private/Online/AuthOSSAdapter.cpp
--- a/Engine/Plugins/Online/OnlineServicesOSSAdapter/Source/Private/Online/AuthOSSAdapter.cpp
+++ b/Engine/Plugins/Online/OnlineServicesOSSAdapter/Source/Private/Online/AuthOSSAdapter.cpp
@@ -528,7 +528,7 @@ TOnlineAsyncOpHandle<FAuthQueryExternalAuthToken> FAuthOSSAdapter::QueryExternal
 
 FUniqueNetIdPtr FAuthOSSAdapter::GetUniqueNetId(FAccountId AccountId) const
 {
-	check(AccountId.GetOnlineServicesType() == GetOnlineServicesOSSAdapter().GetServicesProvider());
+	check(!AccountId.IsValid() || AccountId.GetOnlineServicesType() == GetOnlineServicesOSSAdapter().GetServicesProvider());
 
 	return GetOnlineServicesOSSAdapter().GetAccountIdRegistry().GetIdValue(AccountId);
 }
diff --git a/Engine/Plugins/Online/OnlineServicesOSSAdapter/Source/Private/Online/OnlineServicesOSSAdapter.cpp b/Engine/Plugins/Online/OnlineServicesOSSAdapter/Source/Private/Online/OnlineServicesOSSAdapter.cpp
--- a/Engine/Plugins/Online/OnlineServicesOSSAdapter/Source/Private/Online/OnlineServicesOSSAdapter.cpp
+++ b/Engine/Plugins/Online/OnlineServicesOSSAdapter/Source/Private/Online/OnlineServicesOSSAdapter.cpp
@@ -41,7 +41,11 @@ void FOnlineServicesOSSAdapter::RegisterComponents()
 
 	Components.Register<FPresenceOSSAdapter>(*this);
 	Components.Register<FPrivilegesOSSAdapter>(*this);
-	Components.Register<FStatsOSSAdapter>(*this);
+
+	if (Subsystem->GetStatsInterface().IsValid())
+	{
+		Components.Register<FStatsOSSAdapter>(*this);
+	}
 
 	if (Subsystem->GetFriendsInterface().IsValid())
 	{
