Replacing a Warning with Verbose

#prefight 6519a8f0c779a827b14c0e1e

diff --git a/Engine/Plugins/Online/OnlineServices/Source/OnlineServicesCommonEngineUtils/Private/Online/SessionsLAN.cpp b/Engine/Plugins/Online/OnlineServices/Source/OnlineServicesCommonEngineUtils/Private/Online/SessionsLAN.cpp
--- a/Engine/Plugins/Online/OnlineServices/Source/OnlineServicesCommonEngineUtils/Private/Online/SessionsLAN.cpp
+++ b/Engine/Plugins/Online/OnlineServices/Source/OnlineServicesCommonEngineUtils/Private/Online/SessionsLAN.cpp
@@ -409,7 +409,7 @@ void FSessionsLAN::OnLANSearchTimeout(const FAccountId LocalAccountId)
 		LANSessionManager->StopLANSession();
 	}
 
-	UE_LOG(LogOnlineServices, Warning, TEXT("[FSessionsLAN::OnLANSearchTimeout] %d sessions found!"), SearchResultsUserMap.FindChecked(LocalAccountId).Num());
+	UE_LOG(LogOnlineServices, Verbose, TEXT("[FSessionsLAN::OnLANSearchTimeout] %d sessions found!"), SearchResultsUserMap.FindChecked(LocalAccountId).Num());
 
 	TArray<FOnlineSessionId>& FoundSessionIds = SearchResultsUserMap.FindChecked(LocalAccountId);
 
