OnlineServicesCommon - Add MakeExecHandler helpers

diff --git a/Engine/Plugins/Online/OnlineServices/Source/OnlineServicesCommon/Public/Online/OnlineExecHandler.h b/Engine/Plugins/Online/OnlineServices/Source/OnlineServicesCommon/Public/Online/OnlineExecHandler.h
--- a/Engine/Plugins/Online/OnlineServices/Source/OnlineServicesCommon/Public/Online/OnlineExecHandler.h
+++ b/Engine/Plugins/Online/OnlineServices/Source/OnlineServicesCommon/Public/Online/OnlineExecHandler.h
@@ -6,6 +6,7 @@
 #include "Online/OnlineAsyncOpHandle.h"
 #include "Online/OnlineServicesLog.h"
 #include "Online/OnlineResult.h"
+#include "Templates/Function.h"
 
 #include "Online/Auth.h"
 
@@ -948,6 +949,7 @@ private:
 	MemberFunctionPtrType Function;
 };
 
+/*** Creates an exec handler for an object that provides Exec and Help */
 template <typename T>
 class TOnlineComponentExecHandler : public IOnlineExecHandler
 {
@@ -957,7 +959,7 @@ public:
 	{
 	}
 
-	virtual bool Exec(UWorld* World, const TCHAR* Cmd, FOutputDevice& Ar)
+	virtual bool Exec(UWorld* World, const TCHAR* Cmd, FOutputDevice& Ar) override
 	{
 		return Component->Exec(World, Cmd, Ar);
 	}
@@ -971,6 +973,85 @@ private:
 	T* Component;
 };
 
+/** Creates an exec handler for anything that can be stored in a TFunction */
+class FOnlineFunctionExecHandler : public IOnlineExecHandler
+{
+public:
+	using FunctionType = TFunction<bool(UWorld*,const TCHAR*,FOutputDevice&)>;
+
+	FOnlineFunctionExecHandler(FunctionType&& InFunction, FString&& InHelpString)
+		: Function(MoveTemp(InFunction))
+		, HelpString(MoveTemp(InHelpString))
+	{
+	}
+
+	virtual bool Exec(UWorld* World, const TCHAR* Cmd, FOutputDevice& Ar) override
+	{
+		if (Function)
+		{
+			return Function(World, Cmd, Ar);
+		}
+		else
+		{
+			return false;
+		}
+	}
+
+	virtual bool Help(UWorld* World, const TCHAR* Cmd, FOutputDevice& Ar) override
+	{
+		Ar.Log(HelpString);
+		return true;
+	}
+
+private:
+	FunctionType Function;
+	FString HelpString;
+};
+
+/** Creates an exec handler that calls a class member function */
+template <typename T>
+class TOnlineMemberFunctionExecHandler : public IOnlineExecHandler
+{
+public:
+	using MemberFunctionPtrType = bool(T::*)(UWorld*,const TCHAR*,FOutputDevice&);
+
+	TOnlineMemberFunctionExecHandler(T* InObject, MemberFunctionPtrType InFunction, FString&& InHelpString)
+		: Object(InObject)
+		, Function(InFunction)
+		, HelpString(MoveTemp(InHelpString))
+	{
+	}
+
+	virtual bool Exec(UWorld* World, const TCHAR* Cmd, FOutputDevice& Ar) override
+	{
+		return Invoke(Function, Object, World, Cmd, Ar);
+	}
+
+	virtual bool Help(UWorld* World, const TCHAR* Cmd, FOutputDevice& Ar) override
+	{
+		Ar.Log(HelpString);
+		return true;
+	}
+
+private:
+	T* Object;
+	MemberFunctionPtrType Function;
+	FString HelpString;
+};
+
+/** Creates an exec handler for anything that can be stored in a TFunction */
+inline TUniquePtr<IOnlineExecHandler> MakeExecHandler(TFunction<bool(UWorld*, const TCHAR*, FOutputDevice&)>&& Function, FString&& HelpString = FString())
+{
+	return MakeUnique<FOnlineFunctionExecHandler>(MoveTempIfPossible(Function), MoveTemp(HelpString));
+}
+
+/** Creates an exec handler that calls a class member function */
+template <typename T>
+TUniquePtr<IOnlineExecHandler> MakeExecHandler(T* Object, bool(T::*Function)(UWorld*, const TCHAR*, FOutputDevice&), FString&& HelpString = FString())
+{
+	return MakeUnique<TOnlineMemberFunctionExecHandler<T>>(Object, Function, MoveTemp(HelpString));
+}
+
 /* UE::Online */ }
 
 #if UE_ENABLE_INCLUDE_ORDER_DEPRECATED_IN_5_2
