Fix static analysis warning: warning V591: Non-void function should return a value.

diff --git a/Engine/Plugins/Online/OnlineServices/Source/OnlineServicesCommon/Public/Online/OnlineConfig.h b/Engine/Plugins/Online/OnlineServices/Source/OnlineServicesCommon/Public/Online/OnlineConfig.h
--- a/Engine/Plugins/Online/OnlineServices/Source/OnlineServicesCommon/Public/Online/OnlineConfig.h
+++ b/Engine/Plugins/Online/OnlineServices/Source/OnlineServicesCommon/Public/Online/OnlineConfig.h
@@ -330,8 +330,9 @@ bool LoadConfigStructMemberValue(IOnlineConfigStruct& ConfigStruct, const TCHAR*
 	{
 		return ConfigStruct.GetValue(Key, Value);
 	}
-	else if constexpr (std::is_same_v<decltype(ConfigStruct.GetValue(Key, Value)), int32>)
+	else
 	{
+		static_assert(std::is_same_v<decltype(ConfigStruct.GetValue(Key, Value)), int32>, "IOnlineConfigStruct::GetValue must return bool or int32");
 		return ConfigStruct.GetValue(Key, Value) > 0;
 	}
 }
@@ -343,8 +344,9 @@ bool LoadConfigValue(IOnlineConfigProvider& Provider, const TCHAR* Section, cons
 	{
 		return Provider.GetValue(Section, Key, Value);
 	}
-	else if constexpr (std::is_same_v<decltype(Provider.GetValue(Section, Key, Value)), int32>)
+	else 
 	{
+		static_assert(std::is_same_v<decltype(Provider.GetValue(Section, Key, Value)), int32>, "IOnlineConfigProvider::GetValue must return bool or int32");
 		return Provider.GetValue(Section, Key, Value) > 0;
 	}
 }
