Fix UE_ONLINE_ERROR_EXTERNAL

diff --git a/Engine/Plugins/Online/OnlineServices/Source/OnlineServicesInterface/Public/Online/OnlineErrorCode.h b/Engine/Plugins/Online/OnlineServices/Source/OnlineServicesInterface/Public/Online/OnlineErrorCode.h
--- a/Engine/Plugins/Online/OnlineServices/Source/OnlineServicesInterface/Public/Online/OnlineErrorCode.h
+++ b/Engine/Plugins/Online/OnlineServices/Source/OnlineServicesInterface/Public/Online/OnlineErrorCode.h
@@ -78,7 +78,7 @@ namespace ErrorCode { namespace Category { \
 			{ \
 				InnerPtr = MakeShared<UE::Online::FOnlineError, ESPMode::ThreadSafe>(Inner.GetValue()); \
 			} \
-			return UE::Online::FOnlineError( UE::Online::Errors::ErrorCode::Create(UE::Online::Errors::ErrorCode::Category::CategoryName##_System, UE::Online::Errors::ErrorCode::Category::CategoryName, ErrorCodeValue), MakeShared<UE::Online::FOnlineErrorDetails, ESPMode::ThreadSafe>(FString(#CategoryName"."#Name), ErrorMessage, ErrorText), InnerPtr); \
+			return UE::Online::FOnlineError( UE::Online::Errors::ErrorCode::Create(ErrorCode::Category::CategoryName##_System, ErrorCode::Category::CategoryName, ErrorCodeValue), MakeShared<UE::Online::FOnlineErrorDetails, ESPMode::ThreadSafe>(FString(#CategoryName"."#Name), ErrorMessage, ErrorText), InnerPtr); \
 		}
 
 /** Macro to define an error within a category. Must be used within the UE::Online::Errors namespace. The error will be accessible as UE::Online::Errors::CategoryName::Name() */
