- Fix "expression always false" caught by new version of PVS

diff --git a/Engine/Plugins/Online/OnlineServices/Source/OnlineServicesCommon/Private/Online/OnlineConfig.cpp b/Engine/Plugins/Online/OnlineServices/Source/OnlineServicesCommon/Private/Online/OnlineConfig.cpp
--- a/Engine/Plugins/Online/OnlineServices/Source/OnlineServicesCommon/Private/Online/OnlineConfig.cpp
+++ b/Engine/Plugins/Online/OnlineServices/Source/OnlineServicesCommon/Private/Online/OnlineConfig.cpp
@@ -201,7 +201,7 @@ TArray<FString> FOnlineConfigStructGConfig::CreateValueArray(const FString& InCo
 		if (*Str == TCHAR('('))
 		{
 			Str++;
-			while (*Str != TCHAR(')'))
+			while (*Str && *Str != TCHAR(')'))
 			{
 				const TCHAR* Start = Str;
 
