Add RelyingParty parameter to IAuth::QueryExternalAuthToken.
Add IAuth::GetRelyingParty to access the configured relying party for an Auth interface.

diff --git a/Engine/Plugins/Online/EOSShared/Source/EOSShared/Public/IEOSSDKManager.h b/Engine/Plugins/Online/EOSShared/Source/EOSShared/Public/IEOSSDKManager.h
--- a/Engine/Plugins/Online/EOSShared/Source/EOSShared/Public/IEOSSDKManager.h
+++ b/Engine/Plugins/Online/EOSShared/Source/EOSShared/Public/IEOSSDKManager.h
@@ -21,6 +21,7 @@ struct FEOSSDKPlatformConfig
 	FString ClientId;
 	FString ClientSecret;
 	FString EncryptionKey;
+	FString RelyingPartyURI;
 	FString OverrideCountryCode;
 	FString OverrideLocaleCode;
 	FString DeploymentId;
diff --git a/Engine/Plugins/Online/OnlineServices/Source/OnlineServicesCommon/Private/Online/AuthCommon.cpp b/Engine/Plugins/Online/OnlineServices/Source/OnlineServicesCommon/Private/Online/AuthCommon.cpp
--- a/Engine/Plugins/Online/OnlineServices/Source/OnlineServicesCommon/Private/Online/AuthCommon.cpp
+++ b/Engine/Plugins/Online/OnlineServices/Source/OnlineServicesCommon/Private/Online/AuthCommon.cpp
@@ -182,6 +182,11 @@ TOnlineResult<FAuthGetAllLocalOnlineUsers> FAuthCommon::GetAllLocalOnlineUsers(F
 		GetAccountInfoRegistry().GetAllAccountInfo([](const TSharedRef<FAccountInfo>& AccountInfo) { return IsOnlineStatus(AccountInfo->LoginStatus); }) });
 }
 
+TOnlineResult<FAuthGetRelyingParty> FAuthCommon::GetRelyingParty() const
+{
+	return TOnlineResult<FAuthGetRelyingParty>(Errors::NotImplemented());
+}
+
 TOnlineEvent<void(const FAuthLoginStatusChanged&)> FAuthCommon::OnLoginStatusChanged()
 {
 	return OnAuthLoginStatusChangedEvent;
@@ -203,4 +208,10 @@ bool FAuthCommon::IsLoggedIn(const FAccountId& AccountId) const
 	return (AccountInfo && IsOnlineStatus(AccountInfo->LoginStatus));
 }
 
+bool FAuthCommon::IsLoggedIn(const FPlatformUserId& PlatformUserId) const
+{
+	const TSharedPtr<FAccountInfo> AccountInfo = GetAccountInfoRegistry().Find(PlatformUserId);
+	return (AccountInfo && IsOnlineStatus(AccountInfo->LoginStatus));
+}
+
 /* UE::Online */ }
diff --git a/Engine/Plugins/Online/OnlineServices/Source/OnlineServicesCommon/Public/Online/AuthCommon.h b/Engine/Plugins/Online/OnlineServices/Source/OnlineServicesCommon/Public/Online/AuthCommon.h
--- a/Engine/Plugins/Online/OnlineServices/Source/OnlineServicesCommon/Public/Online/AuthCommon.h
+++ b/Engine/Plugins/Online/OnlineServices/Source/OnlineServicesCommon/Public/Online/AuthCommon.h
@@ -55,10 +55,12 @@ public:
 	virtual TOnlineResult<FAuthGetLocalOnlineUserByOnlineAccountId> GetLocalOnlineUserByOnlineAccountId(FAuthGetLocalOnlineUserByOnlineAccountId::Params&& Params) const override;
 	virtual TOnlineResult<FAuthGetLocalOnlineUserByPlatformUserId> GetLocalOnlineUserByPlatformUserId(FAuthGetLocalOnlineUserByPlatformUserId::Params&& Params) const override;
 	virtual TOnlineResult<FAuthGetAllLocalOnlineUsers> GetAllLocalOnlineUsers(FAuthGetAllLocalOnlineUsers::Params&& Params) const override;
+	virtual TOnlineResult<FAuthGetRelyingParty> GetRelyingParty() const override;
 	virtual TOnlineEvent<void(const FAuthLoginStatusChanged&)> OnLoginStatusChanged() override;
 	virtual TOnlineEvent<void(const FAuthPendingAuthExpiration&)> OnPendingAuthExpiration() override;
 	virtual TOnlineEvent<void(const FAuthAccountAttributesChanged&)> OnAccountAttributesChanged() override;
 	virtual bool IsLoggedIn(const FAccountId& AccountId) const override;
+	virtual bool IsLoggedIn(const FPlatformUserId& PlatformUserId) const override;
 
 protected:
 	virtual const FAccountInfoRegistry& GetAccountInfoRegistry() const = 0;
diff --git a/Engine/Plugins/Online/OnlineServices/Source/OnlineServicesInterface/Public/Online/Auth.h b/Engine/Plugins/Online/OnlineServices/Source/OnlineServicesInterface/Public/Online/Auth.h
--- a/Engine/Plugins/Online/OnlineServices/Source/OnlineServicesInterface/Public/Online/Auth.h
+++ b/Engine/Plugins/Online/OnlineServices/Source/OnlineServicesInterface/Public/Online/Auth.h
@@ -234,6 +234,8 @@ struct FAuthQueryExternalAuthToken
 		FAccountId LocalAccountId;
 		/** The method of external auth to provide. */
 		EExternalAuthTokenMethod Method = EExternalAuthTokenMethod::Primary;
+		/** The Relying Party for the external auth token. */
+		FString RelyingParty;
 	};
 
 	struct Result
@@ -358,6 +360,20 @@ struct FAuthGetAllLocalOnlineUsers
 	};
 };
 
+struct FAuthGetRelyingParty
+{
+	static constexpr TCHAR Name[] = TEXT("GetRelyingParty");
+
+	struct Params
+	{
+	};
+
+	struct Result
+	{
+		FString RelyingParty;
+	};
+};
+
 /** Struct for LoginStatusChanged event */
 struct FAuthLoginStatusChanged
 {
@@ -478,6 +494,11 @@ public:
 	 */
 	virtual TOnlineResult<FAuthGetAllLocalOnlineUsers> GetAllLocalOnlineUsers(FAuthGetAllLocalOnlineUsers::Params&& Params) const = 0;
 
+	/**
+	 * Retrieve the relying party to use when fetching an external token to authenticate with this interface.
+	 */
+	virtual TOnlineResult<FAuthGetRelyingParty> GetRelyingParty() const = 0;
+
 	/**
 	 * Triggered when the login status for a logged in user changes.
 	 */
@@ -499,6 +520,11 @@ public:
 	 * Helper for querying the login status of a local user.
 	 */
 	virtual bool IsLoggedIn(const FAccountId& AccountId) const = 0;
+
+	/**
+	 * Helper for querying the login status of a local user by platform user id.
+	 */
+	virtual bool IsLoggedIn(const FPlatformUserId& PlatformUserId) const = 0;
 };
 
 namespace Meta {
@@ -633,6 +659,13 @@ BEGIN_ONLINE_STRUCT_META(FAuthGetAllLocalOnlineUsers::Result)
 	ONLINE_STRUCT_FIELD(FAuthGetAllLocalOnlineUsers::Result, AccountInfo)
 END_ONLINE_STRUCT_META()
 
+BEGIN_ONLINE_STRUCT_META(FAuthGetRelyingParty::Params)
+END_ONLINE_STRUCT_META()
+
+BEGIN_ONLINE_STRUCT_META(FAuthGetRelyingParty::Result)
+	ONLINE_STRUCT_FIELD(FAuthGetRelyingParty::Result, RelyingParty)
+END_ONLINE_STRUCT_META()
+
 BEGIN_ONLINE_STRUCT_META(FAuthLoginStatusChanged)
 	ONLINE_STRUCT_FIELD(FAuthLoginStatusChanged, AccountInfo),
 	ONLINE_STRUCT_FIELD(FAuthLoginStatusChanged, LoginStatus)
diff --git a/Engine/Plugins/Online/OnlineServicesEOSGS/Source/Private/Online/AuthEOSGS.cpp b/Engine/Plugins/Online/OnlineServicesEOSGS/Source/Private/Online/AuthEOSGS.cpp
--- a/Engine/Plugins/Online/OnlineServicesEOSGS/Source/Private/Online/AuthEOSGS.cpp
+++ b/Engine/Plugins/Online/OnlineServicesEOSGS/Source/Private/Online/AuthEOSGS.cpp
@@ -807,6 +807,28 @@ TOnlineAsyncOpHandle<FAuthBeginVerifiedAuthSession> FAuthEOSGS::BeginVerifiedAut
 	return Operation->GetHandle();
 }
 
+TOnlineResult<FAuthGetRelyingParty> FAuthEOSGS::GetRelyingParty() const
+{
+	IEOSSDKManager* Manager = IEOSSDKManager::Get();
+	const FString& PlatformConfigName = GetServices<FOnlineServicesEOSGS>().GetEOSPlatformHandle()->GetConfigName();
+	const FEOSSDKPlatformConfig* Config = Manager ? Manager->GetPlatformConfig(PlatformConfigName) : nullptr;
+	if (Config == nullptr)
+	{
+		UE_LOG(LogOnlineServices, Error, TEXT("[%hs]: No configuration found for EOS platform %s."), __FUNCTION__, *PlatformConfigName);
+		return TOnlineResult<FAuthGetRelyingParty>(Errors::NotConfigured());
+	}
+
+	FString OutRelyingParty = Config->RelyingPartyURI;
+	if (OutRelyingParty.IsEmpty())
+	{
+		UE_LOG(LogOnlineServices, Warning, TEXT("[%hs]: Relying party configuration is missing for sandbox %s"), __FUNCTION__, *Config->SandboxId);
+		return TOnlineResult<FAuthGetRelyingParty>(Errors::NotConfigured());
+	}
+
+	UE_LOG(LogOnlineServices, VeryVerbose, TEXT("[%hs]: Using relying party: %s"), __FUNCTION__, *OutRelyingParty);
+	return TOnlineResult<FAuthGetRelyingParty>(FAuthGetRelyingParty::Result{MoveTemp(OutRelyingParty)});
+}
+
 TOnlineAsyncOpHandle<FAuthEndVerifiedAuthSession> FAuthEOSGS::EndVerifiedAuthSession(FAuthEndVerifiedAuthSession::Params&& Params)
 {
 	// Todo
diff --git a/Engine/Plugins/Online/OnlineServicesEOSGS/Source/Public/Online/AuthEOSGS.h b/Engine/Plugins/Online/OnlineServicesEOSGS/Source/Public/Online/AuthEOSGS.h
--- a/Engine/Plugins/Online/OnlineServicesEOSGS/Source/Public/Online/AuthEOSGS.h
+++ b/Engine/Plugins/Online/OnlineServicesEOSGS/Source/Public/Online/AuthEOSGS.h
@@ -205,6 +205,7 @@ public:
 	virtual TOnlineAsyncOpHandle<FAuthCancelVerifiedAuthTicket> CancelVerifiedAuthTicket(FAuthCancelVerifiedAuthTicket::Params&& Params) override;
 	virtual TOnlineAsyncOpHandle<FAuthBeginVerifiedAuthSession> BeginVerifiedAuthSession(FAuthBeginVerifiedAuthSession::Params&& Params) override;
 	virtual TOnlineAsyncOpHandle<FAuthEndVerifiedAuthSession> EndVerifiedAuthSession(FAuthEndVerifiedAuthSession::Params&& Params) override;
+	virtual TOnlineResult<FAuthGetRelyingParty> GetRelyingParty() const override;
 	// End IAuth
 
 	// Begin FAuthEOSGS
