Fix mismatched AllowWindowsPlatformTypes header.

diff --git a/Engine/Plugins/Online/OnlineServices/Source/OnlineServicesInterface/Public/Online/Windows/WindowsOnlineErrorDefinitions.h b/Engine/Plugins/Online/OnlineServices/Source/OnlineServicesInterface/Public/Online/Windows/WindowsOnlineErrorDefinitions.h
--- a/Engine/Plugins/Online/OnlineServices/Source/OnlineServicesInterface/Public/Online/Windows/WindowsOnlineErrorDefinitions.h
+++ b/Engine/Plugins/Online/OnlineServices/Source/OnlineServicesInterface/Public/Online/Windows/WindowsOnlineErrorDefinitions.h
@@ -25,7 +25,7 @@ namespace UE::Online::Errors {
 	UE_ONLINE_ERROR(Windows, Unexpected,	 	0x8000FFFF, TEXT("E_UNEXPECTED"),		LOCTEXT("E_UNEXPECTED", "Unexpected failure"))
 
 
-	FString Internal_HResultToString(HRESULT Result)
+	inline FString Internal_HResultToString(HRESULT Result)
 	{
 		TCHAR Buffer[1024];
 		Buffer[0] = TEXT('\0');
@@ -76,4 +76,6 @@ namespace UE::Online::Errors {
 
 } /* namespace UE::Online::Errors */
 
+#include "Windows/HideWindowsPlatformTypes.h"
+
 #undef LOCTEXT_NAMESPACE
\ No newline at end of file
