Fix more JSon issues

diff --git a/Engine/Plugins/Online/OnlineServices/Source/OnlineServicesInterface/OnlineServicesInterface.Build.cs b/Engine/Plugins/Online/OnlineServices/Source/OnlineServicesInterface/OnlineServicesInterface.Build.cs
--- a/Engine/Plugins/Online/OnlineServices/Source/OnlineServicesInterface/OnlineServicesInterface.Build.cs
+++ b/Engine/Plugins/Online/OnlineServices/Source/OnlineServicesInterface/OnlineServicesInterface.Build.cs
@@ -21,5 +21,7 @@ public class OnlineServicesInterface : ModuleRules
 				"Core",
 			}
 		);
+
+		bAllowAutoRTFMInstrumentation = true;
 	}
 }
