Add Tags parameter to IAuth::LinkAccount.
Add LexTryParseString methods for IAuth enums.

diff --git a/Engine/Plugins/Online/OnlineServices/Source/OnlineServicesInterface/Private/Online/Auth.cpp b/Engine/Plugins/Online/OnlineServices/Source/OnlineServicesInterface/Private/Online/Auth.cpp
--- a/Engine/Plugins/Online/OnlineServices/Source/OnlineServicesInterface/Private/Online/Auth.cpp
+++ b/Engine/Plugins/Online/OnlineServices/Source/OnlineServicesInterface/Private/Online/Auth.cpp
@@ -49,9 +49,9 @@ namespace AccountAttributeData
 const FSchemaAttributeId DisplayName = TEXT("DisplayName");
 }
 
-const TCHAR* LexToString(ELoginStatus Status)
+const TCHAR* LexToString(ELoginStatus Value)
 {
-	switch (Status)
+	switch (Value)
 	{
 	case ELoginStatus::UsingLocalProfile:	return TEXT("UsingLocalProfile");
 	case ELoginStatus::LoggedIn:			return TEXT("LoggedIn");
@@ -61,34 +61,28 @@ const TCHAR* LexToString(ELoginStatus Status)
 	}
 }
 
-void LexFromString(ELoginStatus& OutStatus, const TCHAR* InStr)
+bool LexTryParseString(ELoginStatus& OutValue, const TCHAR* InStr)
 {
-	if (FCString::Stricmp(InStr, TEXT("LoggedIn")) == 0)
-	{
-		OutStatus = ELoginStatus::LoggedIn;
-	}
-	else if (FCString::Stricmp(InStr, TEXT("UsingLocalProfile")) == 0)
-	{
-		OutStatus = ELoginStatus::UsingLocalProfile;
-	}
-	else if (FCString::Stricmp(InStr, TEXT("LoggedInReducedFunctionality")) == 0)
-	{
-		OutStatus = ELoginStatus::LoggedInReducedFunctionality;
-	}
-	else if (FCString::Stricmp(InStr, TEXT("NotLoggedIn")) == 0)
-	{
-		OutStatus = ELoginStatus::NotLoggedIn;
-	}
-	else
+#define ENUM_CASE_FROM_STRING(Enum) if (FCString::Stricmp(InStr, TEXT(#Enum)) == 0) { OutValue = ELoginStatus::Enum; return true; }
+	ENUM_CASE_FROM_STRING(LoggedIn);
+	ENUM_CASE_FROM_STRING(UsingLocalProfile);
+	ENUM_CASE_FROM_STRING(LoggedInReducedFunctionality);
+	ENUM_CASE_FROM_STRING(NotLoggedIn);
+#undef ENUM_CASE_FROM_STRING
+	return false;
+}
+
+void LexFromString(ELoginStatus& OutValue, const TCHAR* InStr)
+{
+	if (!ensureAlwaysMsgf(LexTryParseString(OutValue, InStr), TEXT("Unable to parse ELoginStatus value: %s"), InStr))
 	{
-		checkNoEntry();
-		OutStatus = ELoginStatus::NotLoggedIn;
+		OutValue = ELoginStatus::NotLoggedIn;
 	}
 }
 
-const TCHAR* LexToString(ERemoteAuthTicketAudience Audience)
+const TCHAR* LexToString(ERemoteAuthTicketAudience Value)
 {
-	switch (Audience)
+	switch (Value)
 	{
 	case ERemoteAuthTicketAudience::DedicatedServer:	return TEXT("DedicatedServer");
 	default:											checkNoEntry(); // Intentional fallthrough
@@ -96,26 +90,26 @@ const TCHAR* LexToString(ERemoteAuthTicketAudience Audience)
 	}
 }
 
-void LexFromString(ERemoteAuthTicketAudience& OutAudience, const TCHAR* InStr)
+bool LexTryParseString(ERemoteAuthTicketAudience& OutValue, const TCHAR* InStr)
 {
-	if (FCString::Stricmp(InStr, TEXT("Peer")) == 0)
-	{
-		OutAudience = ERemoteAuthTicketAudience::Peer;
-	}
-	else if (FCString::Stricmp(InStr, TEXT("DedicatedServer")) == 0)
-	{
-		OutAudience = ERemoteAuthTicketAudience::DedicatedServer;
-	}
-	else
+#define ENUM_CASE_FROM_STRING(Enum) if (FCString::Stricmp(InStr, TEXT(#Enum)) == 0) { OutValue = ERemoteAuthTicketAudience::Enum; return true; }
+	ENUM_CASE_FROM_STRING(Peer);
+	ENUM_CASE_FROM_STRING(DedicatedServer);
+#undef ENUM_CASE_FROM_STRING
+	return false;
+}
+
+void LexFromString(ERemoteAuthTicketAudience& OutValue, const TCHAR* InStr)
+{
+	if (!ensureAlwaysMsgf(LexTryParseString(OutValue, InStr), TEXT("Unable to parse EExternalAuthTokenMethod value: %s"), InStr))
 	{
-		checkNoEntry();
-		OutAudience = ERemoteAuthTicketAudience::Peer;
+		OutValue = ERemoteAuthTicketAudience::Peer;
 	}
 }
 
-const TCHAR* LexToString(EExternalAuthTokenMethod Method)
+const TCHAR* LexToString(EExternalAuthTokenMethod Value)
 {
-	switch (Method)
+	switch (Value)
 	{
 	case EExternalAuthTokenMethod::Primary:		return TEXT("Primary");
 	default:									checkNoEntry(); // Intentional fallthrough
@@ -123,20 +117,20 @@ const TCHAR* LexToString(EExternalAuthTokenMethod Method)
 	}
 }
 
-void LexFromString(EExternalAuthTokenMethod& OutMethod, const TCHAR* InStr)
+bool LexTryParseString(EExternalAuthTokenMethod& OutValue, const TCHAR* InStr)
 {
-	if (FCString::Stricmp(InStr, TEXT("Primary")) == 0)
-	{
-		OutMethod = EExternalAuthTokenMethod::Primary;
-	}
-	else if (FCString::Stricmp(InStr, TEXT("Secondary")) == 0)
-	{
-		OutMethod = EExternalAuthTokenMethod::Secondary;
-	}
-	else
+#define ENUM_CASE_FROM_STRING(Enum) if (FCString::Stricmp(InStr, TEXT(#Enum)) == 0) { OutValue = EExternalAuthTokenMethod::Enum; return true; }
+	ENUM_CASE_FROM_STRING(Primary);
+	ENUM_CASE_FROM_STRING(Secondary);
+#undef ENUM_CASE_FROM_STRING
+	return false;
+}
+
+void LexFromString(EExternalAuthTokenMethod& OutValue, const TCHAR* InStr)
+{
+	if (!ensureAlwaysMsgf(LexTryParseString(OutValue, InStr), TEXT("Unable to parse EExternalAuthTokenMethod value: %s"), InStr))
 	{
-		checkNoEntry();
-		OutMethod = EExternalAuthTokenMethod::Primary;
+		OutValue = EExternalAuthTokenMethod::Primary;
 	}
 }
 
diff --git a/Engine/Plugins/Online/OnlineServices/Source/OnlineServicesInterface/Public/Online/Auth.h b/Engine/Plugins/Online/OnlineServices/Source/OnlineServicesInterface/Public/Online/Auth.h
--- a/Engine/Plugins/Online/OnlineServices/Source/OnlineServicesInterface/Public/Online/Auth.h
+++ b/Engine/Plugins/Online/OnlineServices/Source/OnlineServicesInterface/Public/Online/Auth.h
@@ -70,8 +70,9 @@ enum class ELoginStatus : uint8
 	/** Player has been validated by the platform specific authentication service. */
 	LoggedIn
 };
-ONLINESERVICESINTERFACE_API const TCHAR* LexToString(ELoginStatus Status);
-ONLINESERVICESINTERFACE_API void LexFromString(ELoginStatus& OutStatus, const TCHAR* InStr);
+ONLINESERVICESINTERFACE_API const TCHAR* LexToString(ELoginStatus Value);
+ONLINESERVICESINTERFACE_API bool LexTryParseString(ELoginStatus& OutValue, const TCHAR* InStr);
+ONLINESERVICESINTERFACE_API void LexFromString(ELoginStatus& OutValue, const TCHAR* InStr);
 
 inline bool IsOnlineStatus(ELoginStatus LoginStatus)
 {
@@ -88,8 +89,9 @@ enum class ERemoteAuthTicketAudience : uint8
 	 */
 	DedicatedServer,
 };
-ONLINESERVICESINTERFACE_API const TCHAR* LexToString(ERemoteAuthTicketAudience Audience);
-ONLINESERVICESINTERFACE_API void LexFromString(ERemoteAuthTicketAudience& OutAudience, const TCHAR* InStr);
+ONLINESERVICESINTERFACE_API const TCHAR* LexToString(ERemoteAuthTicketAudience Value);
+ONLINESERVICESINTERFACE_API bool LexTryParseString(ERemoteAuthTicketAudience& OutValue, const TCHAR* InStr);
+ONLINESERVICESINTERFACE_API void LexFromString(ERemoteAuthTicketAudience& OutValue, const TCHAR* InStr);
 
 /** Some auth interfaces have more than one method for providing credentials when linking to an
  *  external account. An example usage is when the auth interface can provide a token for linking
@@ -103,8 +105,9 @@ enum class EExternalAuthTokenMethod : uint8
 	/** Acquire an external auth token using the secondary method provided by the auth interface. */
 	Secondary,
 };
-ONLINESERVICESINTERFACE_API const TCHAR* LexToString(EExternalAuthTokenMethod Method);
-ONLINESERVICESINTERFACE_API void LexFromString(EExternalAuthTokenMethod& OutMethod, const TCHAR* InStr);
+ONLINESERVICESINTERFACE_API const TCHAR* LexToString(EExternalAuthTokenMethod Value);
+ONLINESERVICESINTERFACE_API bool LexTryParseString(EExternalAuthTokenMethod& OutValue, const TCHAR* InStr);
+ONLINESERVICESINTERFACE_API void LexFromString(EExternalAuthTokenMethod& OutValue, const TCHAR* InStr);
 
 struct FAccountInfo
 {
@@ -199,6 +202,8 @@ struct FAuthLinkAccount
 		FPlatformUserId PlatformUserId = PLATFORMUSERID_NONE;
 		/** The id associated with external auth data used to continue a previous login attempt. */
 		FLoginContinuationId ContinuationId;
+		/** Implementation dependent tags which affect how the account is linked. */
+		TArray<FName> Tags;
 	};
 
 	struct Result
