Add support to register online services command with FDateTime in the params;
Move the define of GetOpDataChecked and corresponding includes to OnlineUtilsCommon.h, in OnlineServicesCommon itself.

diff --git a/Engine/Plugins/Online/OnlineServices/Source/OnlineServicesCommon/Private/Online/LobbiesCommonTests.cpp b/Engine/Plugins/Online/OnlineServices/Source/OnlineServicesCommon/Private/Online/LobbiesCommonTests.cpp
--- a/Engine/Plugins/Online/OnlineServices/Source/OnlineServicesCommon/Private/Online/LobbiesCommonTests.cpp
+++ b/Engine/Plugins/Online/OnlineServices/Source/OnlineServicesCommon/Private/Online/LobbiesCommonTests.cpp
@@ -5,6 +5,7 @@
 #include "Misc/AutomationTest.h"
 #include "Online/LobbiesCommon.h"
 #include "Online/OnlineUtils.h"
+#include "Online/OnlineUtilsCommon.h"
 
 namespace UE::Online {
 
diff --git a/Engine/Plugins/Online/OnlineServices/Source/OnlineServicesCommon/Public/Online/OnlineUtilsCommon.h b/Engine/Plugins/Online/OnlineServices/Source/OnlineServicesCommon/Public/Online/OnlineUtilsCommon.h
new file mode 100644
--- /dev/null
+++ b/Engine/Plugins/Online/OnlineServices/Source/OnlineServicesCommon/Public/Online/OnlineUtilsCommon.h
@@ -0,0 +1,17 @@
+// Copyright Epic Games, Inc. All Rights Reserved.
+
+#pragma once
+
+#include "Online/OnlineAsyncOp.h"
+
+namespace UE::Online {
+
+template <typename DataType, typename OpType>
+const DataType& GetOpDataChecked(const TOnlineAsyncOp<OpType>& Op, const FString& Key)
+{
+	const DataType* Data = Op.Data.template Get<DataType>(Key);
+	check(Data);
+	return *Data;
+}
+
+/* UE::Online */ }
diff --git a/Engine/Plugins/Online/OnlineServices/Source/OnlineServicesInterface/Private/Online/OnlineUtils.cpp b/Engine/Plugins/Online/OnlineServices/Source/OnlineServicesInterface/Private/Online/OnlineUtils.cpp
new file mode 100644
--- /dev/null
+++ b/Engine/Plugins/Online/OnlineServices/Source/OnlineServicesInterface/Private/Online/OnlineUtils.cpp
@@ -0,0 +1,19 @@
+// Copyright Epic Games, Inc. All Rights Reserved.
+
+#include "Online/OnlineUtils.h"
+#include "Misc/DateTime.h"
+
+namespace UE::Online
+{
+
+FString ToLogString(const FDateTime& Time)
+{
+	return Time.ToString();
+}
+
+void LexFromString(FDateTime& Result, const TCHAR* Input)
+{
+	FDateTime::Parse(Input, Result);
+}
+
+}
diff --git a/Engine/Plugins/Online/OnlineServices/Source/OnlineServicesInterface/Public/Online/OnlineUtils.h b/Engine/Plugins/Online/OnlineServices/Source/OnlineServicesInterface/Public/Online/OnlineUtils.h
--- a/Engine/Plugins/Online/OnlineServices/Source/OnlineServicesInterface/Public/Online/OnlineUtils.h
+++ b/Engine/Plugins/Online/OnlineServices/Source/OnlineServicesInterface/Public/Online/OnlineUtils.h
@@ -2,11 +2,14 @@
 
 #pragma once
 
-// HEADER_UNIT_SKIP - Bad includes
-
-#include "Online/OnlineAsyncOp.h"
 #include "Online/OnlineMeta.h"
+#include "Containers/Map.h"
 #include "Containers/UnrealString.h"
+#include "Internationalization/Text.h"
+#include "Misc/TVariant.h"
+#include "Templates/SharedPointer.h"
+
+struct FDateTime;
 
 namespace UE::Online {
 
@@ -22,7 +25,6 @@ template <typename T, ESPMode Mode> FString ToLogString(const TSharedRef<T, Mode
 inline FString ToLogString(const FString& String);
 inline FString ToLogString(const FName& Name);
 inline FString ToLogString(const FText& Text);
-inline FString ToLogString(const FDateTime& Time);
 inline FString ToLogString(uint8 Value);
 inline FString ToLogString(int8 Value);
 inline FString ToLogString(uint16 Value);
@@ -36,6 +38,8 @@ inline FString ToLogString(float Value);
 inline FString ToLogString(double Value);
 template <typename T> std::enable_if_t<!TModels_V<Meta::COnlineMetadataAvailable, T>, FString> ToLogString(const T& Value);
 template <typename T> std::enable_if_t<TModels_V<Meta::COnlineMetadataAvailable, T>, FString> ToLogString(const T& Value);
+ONLINESERVICESINTERFACE_API FString ToLogString(const FDateTime& Time);
+ONLINESERVICESINTERFACE_API void LexFromString(FDateTime& Result, const TCHAR* Input);
 
 template <typename T>
 FString ToLogString(const TArray<T>& Array)
@@ -117,12 +121,6 @@ inline FString ToLogString(const FText& Text)
 	return Text.ToString();
 }
 
-inline FString ToLogString(const FDateTime& Time)
-{
-	return Time.ToString();
-}
-
-
 inline FString ToLogString(uint8 Value)
 {
 	return FString::Printf(TEXT("%hhu"), Value);
@@ -214,12 +212,4 @@ std::enable_if_t<TModels_V<Meta::COnlineMetadataAvailable, T>, FString> ToLogStr
 	return LogString;
 }
 
-template <typename DataType, typename OpType>
-const DataType& GetOpDataChecked(const TOnlineAsyncOp<OpType>& Op, const FString& Key)
-{
-	const DataType* Data = Op.Data.template Get<DataType>(Key);
-	check(Data);
-	return *Data;
-}
-
 /* UE::Online */ }
\ No newline at end of file
diff --git a/Engine/Plugins/Online/OnlineServicesEOS/Source/Private/Online/AuthEOS.cpp b/Engine/Plugins/Online/OnlineServicesEOS/Source/Private/Online/AuthEOS.cpp
--- a/Engine/Plugins/Online/OnlineServicesEOS/Source/Private/Online/AuthEOS.cpp
+++ b/Engine/Plugins/Online/OnlineServicesEOS/Source/Private/Online/AuthEOS.cpp
@@ -13,6 +13,7 @@
 #include "Online/OnlineIdEOS.h"
 #include "Online/OnlineServicesEOS.h"
 #include "Online/OnlineUtils.h"
+#include "Online/OnlineUtilsCommon.h"
 
 #include "eos_auth.h"
 #include "eos_connect.h"
diff --git a/Engine/Plugins/Online/OnlineServicesEOSGS/Source/Private/Online/AuthEOSGS.cpp b/Engine/Plugins/Online/OnlineServicesEOSGS/Source/Private/Online/AuthEOSGS.cpp
--- a/Engine/Plugins/Online/OnlineServicesEOSGS/Source/Private/Online/AuthEOSGS.cpp
+++ b/Engine/Plugins/Online/OnlineServicesEOSGS/Source/Private/Online/AuthEOSGS.cpp
@@ -12,6 +12,7 @@
 #include "Online/OnlineIdEOSGS.h"
 #include "Online/OnlineServicesEOSGS.h"
 #include "Online/OnlineUtils.h"
+#include "Online/OnlineUtilsCommon.h"
 
 #if __has_include(COMPILED_PLATFORM_HEADER(EOSAuthLoginOptions.h))
 #include COMPILED_PLATFORM_HEADER(EOSAuthLoginOptions.h)
diff --git a/Engine/Plugins/Online/OnlineServicesEOSGS/Source/Private/Online/LobbiesEOSGS.cpp b/Engine/Plugins/Online/OnlineServicesEOSGS/Source/Private/Online/LobbiesEOSGS.cpp
--- a/Engine/Plugins/Online/OnlineServicesEOSGS/Source/Private/Online/LobbiesEOSGS.cpp
+++ b/Engine/Plugins/Online/OnlineServicesEOSGS/Source/Private/Online/LobbiesEOSGS.cpp
@@ -12,6 +12,7 @@
 #include "IEOSSDKManager.h"
 
 #include "Online/OnlineUtils.h"
+#include "Online/OnlineUtilsCommon.h"
 #include "eos_lobby.h"
 
 #define UE_ONLINE_LOBBY_EOS_KEY_NAME_LOBBY_DATA TEXT("LobbyData")
diff --git a/Engine/Plugins/Online/OnlineServicesOSSAdapter/Source/Private/Online/AuthOSSAdapter.cpp b/Engine/Plugins/Online/OnlineServicesOSSAdapter/Source/Private/Online/AuthOSSAdapter.cpp
--- a/Engine/Plugins/Online/OnlineServicesOSSAdapter/Source/Private/Online/AuthOSSAdapter.cpp
+++ b/Engine/Plugins/Online/OnlineServicesOSSAdapter/Source/Private/Online/AuthOSSAdapter.cpp
@@ -10,6 +10,7 @@
 #include "Online/DelegateAdapter.h"
 
 #include "Online/OnlineUtils.h"
+#include "Online/OnlineUtilsCommon.h"
 #include "OnlineSubsystem.h"
 #include "Interfaces/OnlineIdentityInterface.h"
 
