Fix the unity compile error when use FDateTime in Params with RegisterCommand(which invokes LexFromString) in OSSv2. Now it uses ParseOnlineExecParams(const TCHAR*& Cmd, FDateTime& Value, IOnlineServices* Services) overload for online specific use.

diff --git a/Engine/Plugins/Online/OnlineServices/Source/OnlineServicesCommon/Public/Online/OnlineExecHandler.h b/Engine/Plugins/Online/OnlineServices/Source/OnlineServicesCommon/Public/Online/OnlineExecHandler.h
--- a/Engine/Plugins/Online/OnlineServices/Source/OnlineServicesCommon/Public/Online/OnlineExecHandler.h
+++ b/Engine/Plugins/Online/OnlineServices/Source/OnlineServicesCommon/Public/Online/OnlineExecHandler.h
@@ -3,6 +3,7 @@
 #pragma once
 
 #include "HAL/IConsoleManager.h"
+#include "Misc/DateTime.h"
 #include "Online/OnlineAsyncOpHandle.h"
 #include "Online/OnlineServicesLog.h"
 #include "Online/OnlineResult.h"
@@ -115,6 +116,7 @@ inline bool ParseOnlineExecParams(const TCHAR*& Cmd, uint32& Value, IOnlineServi
 inline bool ParseOnlineExecParams(const TCHAR*& Cmd, int64& Value, IOnlineServices* Services = nullptr);
 inline bool ParseOnlineExecParams(const TCHAR*& Cmd, uint64& Value, IOnlineServices* Services = nullptr);
 inline bool ParseOnlineExecParams(const TCHAR*& Cmd, bool& Value, IOnlineServices* Services = nullptr);
+inline bool ParseOnlineExecParams(const TCHAR*& Cmd, FDateTime& Value, IOnlineServices* Services = nullptr);
 template <typename T> inline bool ParseOnlineExecParams(const TCHAR*& Cmd, TArray<T>& Array, IOnlineServices* Services = nullptr);
 template <typename T, typename U> inline bool ParseOnlineExecParams(const TCHAR*& Cmd, TMap<T, U>& Map, IOnlineServices* Services = nullptr);
 template <typename T> inline bool ParseOnlineExecParams(const TCHAR*& Cmd, TSet<T>& Set, IOnlineServices* Services = nullptr);
@@ -370,6 +372,19 @@ inline bool ParseOnlineExecParams(const TCHAR*& Cmd, bool& Value, IOnlineService
 	return true;
 }
 
+inline bool ParseOnlineExecParams(const TCHAR*& Cmd, FDateTime& Value, IOnlineServices* Services)
+{
+	FString Token;
+	if (FParse::Token(Cmd, Token, true))
+	{
+		return FDateTime::Parse(Token, Value);
+	}
+	else
+	{
+		return false;
+	}
+}
+
 inline bool ParseOnlineExecParams(const TCHAR*& Cmd, FPlatformUserId& Value, IOnlineServices* Services)
 {
 	FString Token;
diff --git a/Engine/Plugins/Online/OnlineServices/Source/OnlineServicesInterface/Private/Online/OnlineUtils.cpp b/Engine/Plugins/Online/OnlineServices/Source/OnlineServicesInterface/Private/Online/OnlineUtils.cpp
--- a/Engine/Plugins/Online/OnlineServices/Source/OnlineServicesInterface/Private/Online/OnlineUtils.cpp
+++ b/Engine/Plugins/Online/OnlineServices/Source/OnlineServicesInterface/Private/Online/OnlineUtils.cpp
@@ -11,9 +11,4 @@ FString ToLogString(const FDateTime& Time)
 	return Time.ToString();
 }
 
-void LexFromString(FDateTime& Result, const TCHAR* Input)
-{
-	FDateTime::Parse(Input, Result);
-}
-
 }
diff --git a/Engine/Plugins/Online/OnlineServices/Source/OnlineServicesInterface/Public/Online/OnlineUtils.h b/Engine/Plugins/Online/OnlineServices/Source/OnlineServicesInterface/Public/Online/OnlineUtils.h
--- a/Engine/Plugins/Online/OnlineServices/Source/OnlineServicesInterface/Public/Online/OnlineUtils.h
+++ b/Engine/Plugins/Online/OnlineServices/Source/OnlineServicesInterface/Public/Online/OnlineUtils.h
@@ -39,7 +39,6 @@ inline FString ToLogString(double Value);
 template <typename T> std::enable_if_t<!TModels_V<Meta::COnlineMetadataAvailable, T>, FString> ToLogString(const T& Value);
 template <typename T> std::enable_if_t<TModels_V<Meta::COnlineMetadataAvailable, T>, FString> ToLogString(const T& Value);
 ONLINESERVICESINTERFACE_API FString ToLogString(const FDateTime& Time);
-ONLINESERVICESINTERFACE_API void LexFromString(FDateTime& Result, const TCHAR* Input);
 
 template <typename T>
 FString ToLogString(const TArray<T>& Array)
