Add category system to Online ErrorCode so they compare equal to codes in FOnlineError.
#rnx

diff --git a/Engine/Plugins/Online/OnlineServices/Source/OnlineServicesInterface/Public/Online/OnlineErrorCode.h b/Engine/Plugins/Online/OnlineServices/Source/OnlineServicesInterface/Public/Online/OnlineErrorCode.h
--- a/Engine/Plugins/Online/OnlineServices/Source/OnlineServicesInterface/Public/Online/OnlineErrorCode.h
+++ b/Engine/Plugins/Online/OnlineServices/Source/OnlineServicesInterface/Public/Online/OnlineErrorCode.h
@@ -23,14 +23,9 @@ namespace ErrorCode
 
 	ONLINESERVICESINTERFACE_API FString ToString(ErrorCodeType ErrorCode);
 
-	constexpr ErrorCodeType Create(uint32 Source, uint32 Category, uint32 Code)
+	constexpr ErrorCodeType Create(uint32 System, uint32 Category, uint32 Code)
 	{
-		return ((Source & 0xfull) << 60ull) | ((Category & 0x0fffffffull) << 32ull) | static_cast<ErrorCodeType>(Code);
-	}
-
-	constexpr ErrorCodeType Create(uint32 Category, uint32 Code)
-	{
-		return ((Category & 0x0fffffffull) << 32ull) | static_cast<ErrorCodeType>(Code);
+		return ((System & 0xfull) << 60ull) | ((Category & 0x0fffffffull) << 32ull) | static_cast<ErrorCodeType>(Code);
 	}
 
 	namespace System
@@ -43,7 +38,7 @@ namespace ErrorCode
 
 // These three functions are responsible for reading the parts of the uint64 error representation
 // The uint64 error is formatted as follows (big endian):
-// ssss cccc cccc cccc cccc cccc cccc cccc cccc vvvv vvvv vvvv vvvv vvvv vvvv vvvv vvvv 
+// ssss cccc cccc cccc cccc cccc cccc cccc vvvv vvvv vvvv vvvv vvvv vvvv vvvv vvvv 
 ONLINESERVICESINTERFACE_API uint64 ErrorCodeSystem(ErrorCodeType ErrorCode);
 ONLINESERVICESINTERFACE_API uint64 ErrorCodeCategory(ErrorCodeType ErrorCode);
 ONLINESERVICESINTERFACE_API uint64 ErrorCodeValue(ErrorCodeType ErrorCode); 
@@ -83,19 +78,19 @@ namespace ErrorCode { namespace Category { \
 
 /** Macro to define an error within a category. Must be used within the UE::Online::Errors namespace. The error will be accessible as UE::Online::Errors::CategoryName::Name() */
 #define UE_ONLINE_ERROR(CategoryName, Name, ErrorCodeValue, ErrorMessage, ErrorText) \
-    namespace ErrorCode { namespace CategoryName { static constexpr ErrorCodeType Name = Create(Category::CategoryName, ErrorCodeValue); } } \
+    namespace ErrorCode { namespace CategoryName { static constexpr ErrorCodeType Name = Create(Category::CategoryName##_System, Category::CategoryName, ErrorCodeValue); } } \
 	namespace CategoryName { \
 		UE_ONLINE_ERROR_INTERNAL(CategoryName, Name, ErrorCodeValue, ErrorMessage, ErrorText)\
 	}
 
 /** Macro to define a common error in the UE::Online::Errors namespace. Must be used within the UE::Online::Errors namespace. The error will be accessible as UE::Online::Errors::Name() */
 #define UE_ONLINE_ERROR_COMMON(CategoryName, Name, ErrorCodeValue, ErrorMessage, ErrorText) \
-    namespace ErrorCode { namespace CategoryName { static constexpr ErrorCodeType Name = Create(Category::CategoryName, ErrorCodeValue); } } \
+    namespace ErrorCode { namespace CategoryName { static constexpr ErrorCodeType Name = Create(Category::CategoryName##_System, Category::CategoryName, ErrorCodeValue); } } \
 	UE_ONLINE_ERROR_INTERNAL(CategoryName, Name, ErrorCodeValue, ErrorMessage, ErrorText)
 
 /** Macro to define an error in a user defined namespace. The error will be accessible as CurrentNamespace::Name() */
 #define UE_ONLINE_ERROR_EXTERNAL(CategoryName, Name, ErrorCodeValue, ErrorMessage, ErrorText) \
-    namespace ErrorCode { namespace CategoryName { static constexpr UE::Online::Errors::ErrorCodeType Name = UE::Online::Errors::ErrorCode::Create(ErrorCode::Category::CategoryName, ErrorCodeValue); } } \
+    namespace ErrorCode { namespace CategoryName { static constexpr UE::Online::Errors::ErrorCodeType Name = UE::Online::Errors::ErrorCode::Create(Category::CategoryName##_System, ErrorCode::Category::CategoryName, ErrorCodeValue); } } \
 	UE_ONLINE_ERROR_INTERNAL(CategoryName, Name, ErrorCodeValue, ErrorMessage, ErrorText)
 
 } /* namespace UE::Online::Errors */
