Remove invalid online error code Success value.
#rnx

diff --git a/Engine/Plugins/Online/OnlineServices/Source/OnlineServicesInterface/Private/Online/OnlineError.cpp b/Engine/Plugins/Online/OnlineServices/Source/OnlineServicesInterface/Private/Online/OnlineError.cpp
--- a/Engine/Plugins/Online/OnlineServices/Source/OnlineServicesInterface/Private/Online/OnlineError.cpp
+++ b/Engine/Plugins/Online/OnlineServices/Source/OnlineServicesInterface/Private/Online/OnlineError.cpp
@@ -1,6 +1,7 @@
 // Copyright Epic Games, Inc. All Rights Reserved.
 #include "Online/OnlineError.h"
 #include "Online/OnlineErrorCode.h"
+#include "Online/OnlineErrorDefinitions.h"
 #include "Serialization/CompactBinaryWriter.h"
 
 namespace UE::Online{
@@ -13,17 +14,10 @@ namespace UE::Online{
 		{
 			FString ToString(ErrorCodeType ErrorCode)
 			{
-				const uint64 Source = ErrorCodeSystem(ErrorCode);
+				const uint64 System = ErrorCodeSystem(ErrorCode);
 				const uint64 Category = ErrorCodeCategory(ErrorCode); 
 				const uint64 Code = ErrorCodeValue(ErrorCode); 
-				if (Source == 0)
-				{
-					return FString::Printf(TEXT("%llx.%llx"), Category, Code);
-				}
-				else
-				{
-					return FString::Printf(TEXT("%llx.%llx.%llx"), Source, Category, Code);
-				}
+				return FString::Printf(TEXT("%llx.%llx.%llx"), System, Category, Code);
 			}
 		} /*namespace ErrorCode */
 	} /* namespace Errors */
@@ -77,19 +71,9 @@ namespace UE::Online{
 		bool bHasDetails = false;
 		bool bHasInner = false;
 
-		if (OnlineError == UE::Online::Errors::ErrorCode::Success)
-		{
-			FText SuccessText = NSLOCTEXT("OnlineError", "Success", "Success");
-			ErrorMessageArgs.Add(TEXT("ErrorCode"), SuccessText);
-			Writer.AddString(ANSITEXTVIEW("ErrorCode"), SuccessText.ToString());
-		}
-		else
-		{
-			FString ErrorCodeString = OnlineError.GetErrorId();
-			ErrorMessageArgs.Add(TEXT("ErrorCode"), FText::FromString(ErrorCodeString));
-			Writer.AddString(ANSITEXTVIEW("ErrorCode"), ErrorCodeString);
-			
-		}
+		FString ErrorCodeString = OnlineError.GetErrorId();
+		ErrorMessageArgs.Add(TEXT("ErrorCode"), FText::FromString(ErrorCodeString));
+		Writer.AddString(ANSITEXTVIEW("ErrorCode"), ErrorCodeString);
 
 		if (OnlineError.Details)
 		{
@@ -131,5 +115,10 @@ namespace UE::Online{
 		Writer.EndObject();
 	}
 
+	bool FOnlineError::IsSuccess() const
+	{
+		return ErrorCode == Errors::ErrorCode::Common::Success;
+	}
+
 } /* namespace UE::Online */
 
diff --git a/Engine/Plugins/Online/OnlineServices/Source/OnlineServicesInterface/Public/Online/OnlineError.h b/Engine/Plugins/Online/OnlineServices/Source/OnlineServicesInterface/Public/Online/OnlineError.h
--- a/Engine/Plugins/Online/OnlineServices/Source/OnlineServicesInterface/Public/Online/OnlineError.h
+++ b/Engine/Plugins/Online/OnlineServices/Source/OnlineServicesInterface/Public/Online/OnlineError.h
@@ -73,17 +73,13 @@ public:
 		{
 			return Details->GetText(*this); // alternatively could be done via lookup on an error registry
 		}
-		else if (ErrorCode == Errors::ErrorCode::Success)
-		{
-			return LOCTEXT("Success", "Success"); // success message
-		}
 		else
 		{
 			return FText::FromString(GetErrorId()); // generic error message with code
 		}
 	}
 
-	FString GetLogString(bool bIncludePrefix = true, bool bIncludeSuccess = true) const
+	FString GetLogString(bool bIncludePrefix = true) const
 	{
 #if !NO_LOGGING
 		FString MyLogString = TEXT("");
@@ -100,13 +96,6 @@ public:
 			FString LogString = Details->GetLogString(*this);
 			MyLogString = FString::Printf(TEXT("%s%s"), *LogPrefix, *Details->GetLogString(*this));
 		}
-		else if (ErrorCode == Errors::ErrorCode::Success)
-		{
-			if (bIncludeSuccess)
-			{
-				MyLogString = TEXT("Success");
-			}
-		}
 		else
 		{
 			MyLogString = FString::Printf(TEXT("[%s]"), *GetErrorId()); 
@@ -114,7 +103,7 @@ public:
 
 		if (GetInner() != nullptr)
 		{
-			FString InnerLogStr = GetInner()->GetLogString(false, false);
+			FString InnerLogStr = GetInner()->GetLogString(false);
 			if (!InnerLogStr.IsEmpty())
 			{
 				return FString::Printf(TEXT("%s (%s)"), *MyLogString, *InnerLogStr);
@@ -147,11 +136,7 @@ public:
 
 	FString GetFriendlyErrorCode() const
 	{
-		if (ErrorCode == Errors::ErrorCode::Success)
-		{
-			return TEXT("Success");
-		}
-		else if (Details)
+		if (Details)
 		{
 			return Details->GetFriendlyErrorCode(*this);
 		}
@@ -186,6 +171,8 @@ public:
 		return ErrorCode;
 	}
 
+	bool IsSuccess() const;
+
 private:
 	// TSharedPtr instead of TUniquePtr so that we can copy errors easily
 	TSharedPtr<const IOnlineErrorDetails, ESPMode::ThreadSafe> Details;
diff --git a/Engine/Plugins/Online/OnlineServices/Source/OnlineServicesInterface/Public/Online/OnlineErrorCode.h b/Engine/Plugins/Online/OnlineServices/Source/OnlineServicesInterface/Public/Online/OnlineErrorCode.h
--- a/Engine/Plugins/Online/OnlineServices/Source/OnlineServicesInterface/Public/Online/OnlineErrorCode.h
+++ b/Engine/Plugins/Online/OnlineServices/Source/OnlineServicesInterface/Public/Online/OnlineErrorCode.h
@@ -19,8 +19,6 @@ struct FOnlineErrorCode
 
 namespace ErrorCode
 {
-	static constexpr ErrorCodeType Success = ErrorCodeType(0);
-
 	ONLINESERVICESINTERFACE_API FString ToString(ErrorCodeType ErrorCode);
 
 	constexpr ErrorCodeType Create(uint32 System, uint32 Category, uint32 Code)
