Fix for Windows OnlineError issue.
#rnx

diff --git a/Engine/Plugins/Online/OnlineServices/Source/OnlineServicesInterface/Public/Online/Windows/WindowsOnlineErrorDefinitions.h b/Engine/Plugins/Online/OnlineServices/Source/OnlineServicesInterface/Public/Online/Windows/WindowsOnlineErrorDefinitions.h
--- a/Engine/Plugins/Online/OnlineServices/Source/OnlineServicesInterface/Public/Online/Windows/WindowsOnlineErrorDefinitions.h
+++ b/Engine/Plugins/Online/OnlineServices/Source/OnlineServicesInterface/Public/Online/Windows/WindowsOnlineErrorDefinitions.h
@@ -3,6 +3,7 @@
 
 #include "Online/OnlineError.h"
 #include "Online/OnlineErrorCode.h"
+#include "Online/OnlineErrorDefinitions.h"
 #include "Windows/AllowWindowsPlatformTypes.h"
 #include "winerror.h"
 #include "winbase.h"
@@ -65,7 +66,7 @@ namespace UE::Online::Errors {
 	{
 		if (Result == S_OK)
 		{
-			return ErrorCode::Success;
+			return Success();
 		}
 
 		// This isn't localized, but Windows doesn't even have localized error codes anyway- do we need to add some method for localization? (e.g. try to look up the error first using some table generated by the WINDOWS_ERROR macro and use that error's FText)
