Fix external linking for OnlineError IsSuccess function.
#rnx

diff --git a/Engine/Plugins/Online/OnlineServices/Source/OnlineServicesInterface/Public/Online/OnlineError.h b/Engine/Plugins/Online/OnlineServices/Source/OnlineServicesInterface/Public/Online/OnlineError.h
--- a/Engine/Plugins/Online/OnlineServices/Source/OnlineServicesInterface/Public/Online/OnlineError.h
+++ b/Engine/Plugins/Online/OnlineServices/Source/OnlineServicesInterface/Public/Online/OnlineError.h
@@ -171,7 +171,7 @@ public:
 		return ErrorCode;
 	}
 
-	bool IsSuccess() const;
+	ONLINESERVICESINTERFACE_API bool IsSuccess() const;
 
 private:
 	// TSharedPtr instead of TUniquePtr so that we can copy errors easily
