Fix the backwards compatibility break, after moving UE::Online::GetOpDataChecked from OnlineServicesInterface - OnlineUtils.h to OnlineServicesCommon - OnlineUtilsCommon.h.

diff --git a/Engine/Plugins/Online/OnlineServices/Source/OnlineServicesInterface/Public/Online/OnlineUtils.h b/Engine/Plugins/Online/OnlineServices/Source/OnlineServicesInterface/Public/Online/OnlineUtils.h
--- a/Engine/Plugins/Online/OnlineServices/Source/OnlineServicesInterface/Public/Online/OnlineUtils.h
+++ b/Engine/Plugins/Online/OnlineServices/Source/OnlineServicesInterface/Public/Online/OnlineUtils.h
@@ -8,6 +8,9 @@
 #include "Internationalization/Text.h"
 #include "Misc/TVariant.h"
 #include "Templates/SharedPointer.h"
+#if UE_ENABLE_INCLUDE_ORDER_DEPRECATED_IN_5_5
+#include "Online/OnlineUtilsCommon.h"
+#endif
 
 struct FDateTime;
 
