Fix the compile error when OnlineServicesCommon module is not in the included path, compiler can't find the moved OnlineUtilsCommon.h.

diff --git a/Engine/Plugins/Online/OnlineServices/Source/OnlineServicesInterface/OnlineServicesInterface.Build.cs b/Engine/Plugins/Online/OnlineServices/Source/OnlineServicesInterface/OnlineServicesInterface.Build.cs
--- a/Engine/Plugins/Online/OnlineServices/Source/OnlineServicesInterface/OnlineServicesInterface.Build.cs
+++ b/Engine/Plugins/Online/OnlineServices/Source/OnlineServicesInterface/OnlineServicesInterface.Build.cs
@@ -1,6 +1,7 @@
 // Copyright Epic Games, Inc. All Rights Reserved.
 
 using UnrealBuildTool;
+using System.IO;
 
 public class OnlineServicesInterface : ModuleRules
 {
@@ -21,5 +22,8 @@ public class OnlineServicesInterface : ModuleRules
 				"Core",
 			}
 		);
+
+		// For UE_ENABLE_INCLUDE_ORDER_DEPRECATED_IN_5_5: Online/OnlineUtilsCommon.h
+		PublicIncludePaths.Add(Path.Combine(ModuleDirectory, "..", "OnlineServicesCommon", "Public"));
 	}
 }
