Fix comment about UE::Online::IsLoaded

diff --git a/Engine/Plugins/Online/OnlineServices/Source/OnlineServicesInterface/Public/Online/OnlineServices.h b/Engine/Plugins/Online/OnlineServices/Source/OnlineServicesInterface/Public/Online/OnlineServices.h
--- a/Engine/Plugins/Online/OnlineServices/Source/OnlineServicesInterface/Public/Online/OnlineServices.h
+++ b/Engine/Plugins/Online/OnlineServices/Source/OnlineServicesInterface/Public/Online/OnlineServices.h
@@ -195,9 +195,9 @@ ONLINESERVICESINTERFACE_API int32 GetBuildUniqueId();
 /**
  * Check if an instance of the online service is loaded
  *
- * @param OnlineServices Type of online services to retrieve
- * @param InstanceName Name of the services instance to retrieve
- * @return The services instance or an invalid pointer if the services is unavailable
+ * @param OnlineServices Type of online services for the IOnlineServices instance
+ * @param InstanceName Name of the instance
+ * @return true if the instance is loaded
  */
 ONLINESERVICESINTERFACE_API bool IsLoaded(EOnlineServices OnlineServices = EOnlineServices::Default, FName InstanceName = NAME_None);
 
