Add ExternalLoginType::AccountIdHash to allow login using a non-authoritative value based on an account id.

diff --git a/Engine/Plugins/Online/OnlineServices/Source/OnlineServicesInterface/Private/Online/Auth.cpp b/Engine/Plugins/Online/OnlineServices/Source/OnlineServicesInterface/Private/Online/Auth.cpp
--- a/Engine/Plugins/Online/OnlineServices/Source/OnlineServicesInterface/Private/Online/Auth.cpp
+++ b/Engine/Plugins/Online/OnlineServices/Source/OnlineServicesInterface/Private/Online/Auth.cpp
@@ -36,6 +36,7 @@ const FName ItchioJwt = TEXT("ItchioJwt");
 const FName ItchioKey = TEXT("ItchioKey");
 const FName EpicIdToken = TEXT("EpicIdToken");
 const FName AmazonAccessToken = TEXT("AmazonAccessToken");
+const FName AccountIdHash = TEXT("AccountIdHash");
 }
 
 namespace ExternalServerAuthTicketType
diff --git a/Engine/Plugins/Online/OnlineServices/Source/OnlineServicesInterface/Public/Online/Auth.h b/Engine/Plugins/Online/OnlineServices/Source/OnlineServicesInterface/Public/Online/Auth.h
--- a/Engine/Plugins/Online/OnlineServices/Source/OnlineServicesInterface/Public/Online/Auth.h
+++ b/Engine/Plugins/Online/OnlineServices/Source/OnlineServicesInterface/Public/Online/Auth.h
@@ -46,6 +46,7 @@ ONLINESERVICESINTERFACE_API extern const FName ItchioJwt;
 ONLINESERVICESINTERFACE_API extern const FName ItchioKey;
 ONLINESERVICESINTERFACE_API extern const FName EpicIdToken;
 ONLINESERVICESINTERFACE_API extern const FName AmazonAccessToken;
+ONLINESERVICESINTERFACE_API extern const FName AccountIdHash;
 }
 
 namespace ExternalServerAuthTicketType
