Fixed some 'deprecated' FString usage.

diff --git a/Engine/Plugins/Developer/RenderDocPlugin/Source/RenderDocPlugin/Private/RenderDocPluginModule.cpp b/Engine/Plugins/Developer/RenderDocPlugin/Source/RenderDocPlugin/Private/RenderDocPluginModule.cpp
--- a/Engine/Plugins/Developer/RenderDocPlugin/Source/RenderDocPlugin/Private/RenderDocPluginModule.cpp
+++ b/Engine/Plugins/Developer/RenderDocPlugin/Source/RenderDocPlugin/Private/RenderDocPluginModule.cpp
@@ -631,7 +631,7 @@ FString FRenderDocPluginModule::GetNewestCapture()
 	
 	while (RenderDocAPI->GetCapture(Index, LogFile, &LogPathLength, &Timestamp))
 	{
-		OutString = FString(LogPathLength, ANSI_TO_TCHAR(LogFile));
+		OutString = FString::ConstructFromPtrSize(ANSI_TO_TCHAR(LogFile), LogPathLength);
 
 		Index++;
 	}
diff --git a/Engine/Plugins/Experimental/BackChannel/Source/BackChannel/Private/Utils/DispatchMap.cpp b/Engine/Plugins/Experimental/BackChannel/Source/BackChannel/Private/Utils/DispatchMap.cpp
--- a/Engine/Plugins/Experimental/BackChannel/Source/BackChannel/Private/Utils/DispatchMap.cpp
+++ b/Engine/Plugins/Experimental/BackChannel/Source/BackChannel/Private/Utils/DispatchMap.cpp
@@ -11,7 +11,7 @@ FBackChannelDispatchMap::FBackChannelDispatchMap()
 
 FDelegateHandle FBackChannelDispatchMap::AddRoute(FStringView InPath, FBackChannelRouteDelegate::FDelegate Delegate)
 {
-	FString LowerPath = FString(InPath.Len(), InPath.GetData()).ToLower();
+	FString LowerPath = FString::ConstructFromPtrSize(InPath.GetData(), InPath.Len()).ToLower();
 
 	if (DispatchMap.Contains(LowerPath) == false)
 	{
@@ -23,7 +23,7 @@ FDelegateHandle FBackChannelDispatchMap::AddRoute(FStringView InPath, FBackChann
 
 void FBackChannelDispatchMap::RemoveRoute(FStringView InPath, FDelegateHandle DelegateHandle)
 {
-	FString LowerPath = FString(InPath.Len(), InPath.GetData()).ToLower();
+	FString LowerPath = FString::ConstructFromPtrSize(InPath.GetData(), InPath.Len()).ToLower();
 
 	auto DelegateList = DispatchMap.Find(LowerPath);
 
diff --git a/Engine/Plugins/Experimental/PythonScriptPlugin/Source/PythonScriptPlugin/Private/PyGenUtil.cpp b/Engine/Plugins/Experimental/PythonScriptPlugin/Source/PythonScriptPlugin/Private/PyGenUtil.cpp
--- a/Engine/Plugins/Experimental/PythonScriptPlugin/Source/PythonScriptPlugin/Private/PyGenUtil.cpp
+++ b/Engine/Plugins/Experimental/PythonScriptPlugin/Source/PythonScriptPlugin/Private/PyGenUtil.cpp
@@ -1729,7 +1729,7 @@ bool IsValidName(FStringView InName, FText* OutError)
 		{
 			if (OutError)
 			{
-				*OutError = FText::Format(NSLOCTEXT("PyGenUtil", "InvalidName_RestrictedCharacter", "Name contains '{0}' which is invalid for Python"), FText::AsCultureInvariant(FString(1, &Char)));
+				*OutError = FText::Format(NSLOCTEXT("PyGenUtil", "InvalidName_RestrictedCharacter", "Name contains '{0}' which is invalid for Python"), FText::AsCultureInvariant(FString::ConstructFromPtrSize(&Char, 1)));
 			}
 			return false;
 		}
diff --git a/Engine/Plugins/Media/ElectraCodecs/Source/ElectraDecoders/Private/Utils/MPEG/ElectraUtilsMP4.cpp b/Engine/Plugins/Media/ElectraCodecs/Source/ElectraDecoders/Private/Utils/MPEG/ElectraUtilsMP4.cpp
--- a/Engine/Plugins/Media/ElectraCodecs/Source/ElectraDecoders/Private/Utils/MPEG/ElectraUtilsMP4.cpp
+++ b/Engine/Plugins/Media/ElectraCodecs/Source/ElectraDecoders/Private/Utils/MPEG/ElectraUtilsMP4.cpp
@@ -53,7 +53,7 @@ bool FMP4AtomReader::ReadString(FString& OutString, uint16 NumBytes)
 			return false;
 		}
 		FUTF8ToTCHAR cnv((const ANSICHAR*)Buf.GetData(), NumBytes);
-		OutString = FString(cnv.Length(), cnv.Get());
+		OutString = FString::ConstructFromPtrSize(cnv.Get(), cnv.Length());
 		return true;
 	}
 	return false;
@@ -80,7 +80,7 @@ bool FMP4AtomReader::ReadStringUTF8(FString& OutString, int32 NumBytes)
 	if (ReadBytes(Buf.GetData(), NumBytes))
 	{
 		FUTF8ToTCHAR cnv((const ANSICHAR*)Buf.GetData(), NumBytes);
-		OutString = FString(cnv.Length(), cnv.Get());
+		OutString = FString::ConstructFromPtrSize(cnv.Get(), cnv.Length());
 		return true;
 	}
 	return false;
diff --git a/Engine/Plugins/Media/ElectraPlayer/Source/ElectraPlayerRuntime/Private/Runtime/Demuxer/ParserMKV.cpp b/Engine/Plugins/Media/ElectraPlayer/Source/ElectraPlayerRuntime/Private/Runtime/Demuxer/ParserMKV.cpp
--- a/Engine/Plugins/Media/ElectraPlayer/Source/ElectraPlayerRuntime/Private/Runtime/Demuxer/ParserMKV.cpp
+++ b/Engine/Plugins/Media/ElectraPlayer/Source/ElectraPlayerRuntime/Private/Runtime/Demuxer/ParserMKV.cpp
@@ -724,7 +724,7 @@ namespace Electra
 						}
 						Chars.Add(0);
 						auto Cnv = StringCast<TCHAR>((const ANSICHAR*)Chars.GetData());
-						FString UTF8Text(Cnv.Length(), Cnv.Get());
+						FString UTF8Text = FString::ConstructFromPtrSize(Cnv.Get(), Cnv.Length());
 						*reinterpret_cast<FString*>(Elem->DataValue) = MoveTemp(UTF8Text);
 						break;
 					}
@@ -737,7 +737,7 @@ namespace Electra
 						}
 						Chars.Add(0);
 						auto Cnv = StringCast<TCHAR>((const UTF8CHAR*)Chars.GetData());
-						FString UTF8Text(Cnv.Length(), Cnv.Get());
+						FString UTF8Text = FString::ConstructFromPtrSize(Cnv.Get(), Cnv.Length());
 						*reinterpret_cast<FString*>(Elem->DataValue) = MoveTemp(UTF8Text);
 						break;
 						}
diff --git a/Engine/Plugins/Media/ElectraSubtitles/Source/ElectraSubtitles/Private/ttml/TTMLParser.cpp b/Engine/Plugins/Media/ElectraSubtitles/Source/ElectraSubtitles/Private/ttml/TTMLParser.cpp
--- a/Engine/Plugins/Media/ElectraSubtitles/Source/ElectraSubtitles/Private/ttml/TTMLParser.cpp
+++ b/Engine/Plugins/Media/ElectraSubtitles/Source/ElectraSubtitles/Private/ttml/TTMLParser.cpp
@@ -74,19 +74,19 @@ struct ExpatUtils
 	static FString ToFStringL(const XML_LChar* In)
 	{
 		FUTF8ToTCHAR Converted((const ANSICHAR*)In);
-		return FString(Converted.Length(), (TCHAR*)Converted.Get());
+		return FString::ConstructFromPtrSize((TCHAR*)Converted.Get(), Converted.Length());
 	}
 
 	static FString ToFString(const XML_Char* In)
 	{
 		FUTF8ToTCHAR Converted((const ANSICHAR*)In);
-		return FString(Converted.Length(), (TCHAR*)Converted.Get());
+		return FString::ConstructFromPtrSize((TCHAR*)Converted.Get(), Converted.Length());
 	}
 
 	static FString ToFString(const XML_Char* In, int32 Len)
 	{
 		FUTF8ToTCHAR Converted((const ANSICHAR*)In, Len);
-		return FString(Converted.Length(), (TCHAR*)Converted.Get());
+		return FString::ConstructFromPtrSize((TCHAR*)Converted.Get(), Converted.Length());
 	}
 
 	static TArray<FXmlAttribute> ToAttributes(const XML_Char** InAtts)
diff --git a/Engine/Plugins/Media/ElectraSubtitles/Source/ElectraSubtitles/Private/tx3g/ElectraSubtitleDecoder_TX3G.cpp b/Engine/Plugins/Media/ElectraSubtitles/Source/ElectraSubtitles/Private/tx3g/ElectraSubtitleDecoder_TX3G.cpp
--- a/Engine/Plugins/Media/ElectraSubtitles/Source/ElectraSubtitles/Private/tx3g/ElectraSubtitleDecoder_TX3G.cpp
+++ b/Engine/Plugins/Media/ElectraSubtitles/Source/ElectraSubtitles/Private/tx3g/ElectraSubtitleDecoder_TX3G.cpp
@@ -93,7 +93,7 @@ namespace OptionKeys
 					return false;
 				}
 				FUTF8ToTCHAR cnv((const ANSICHAR*)Buf.GetData(), NumBytes);
-				OutString = FString(cnv.Length(), cnv.Get());
+				OutString = FString::ConstructFromPtrSize(cnv.Get(), cnv.Length());
 				return true;
 			}
 			return false;
diff --git a/Engine/Plugins/Media/ElectraSubtitles/Source/ElectraSubtitles/Private/wvtt/ElectraSubtitleDecoder_WVTT.cpp b/Engine/Plugins/Media/ElectraSubtitles/Source/ElectraSubtitles/Private/wvtt/ElectraSubtitleDecoder_WVTT.cpp
--- a/Engine/Plugins/Media/ElectraSubtitles/Source/ElectraSubtitles/Private/wvtt/ElectraSubtitleDecoder_WVTT.cpp
+++ b/Engine/Plugins/Media/ElectraSubtitles/Source/ElectraSubtitles/Private/wvtt/ElectraSubtitleDecoder_WVTT.cpp
@@ -93,7 +93,7 @@ namespace OptionKeys
 					return false;
 				}
 				FUTF8ToTCHAR cnv((const ANSICHAR*)Buf.GetData(), NumBytes);
-				OutString = FString(cnv.Length(), cnv.Get());
+				OutString = FString::ConstructFromPtrSize(cnv.Get(), cnv.Length());
 				return true;
 			}
 			return false;
diff --git a/Engine/Plugins/Media/ElectraUtil/Source/ElectraBase/Private/StringHelpers.cpp b/Engine/Plugins/Media/ElectraUtil/Source/ElectraBase/Private/StringHelpers.cpp
--- a/Engine/Plugins/Media/ElectraUtil/Source/ElectraBase/Private/StringHelpers.cpp
+++ b/Engine/Plugins/Media/ElectraUtil/Source/ElectraBase/Private/StringHelpers.cpp
@@ -165,7 +165,7 @@ namespace Electra
 		FString ArrayToString(const TArray<uint8>& InArray)
 		{
 			FUTF8ToTCHAR cnv((const ANSICHAR*)InArray.GetData(), InArray.Num());
-			FString UTF8Text(cnv.Length(), cnv.Get());
+			FString UTF8Text = FString::ConstructFromPtrSize(cnv.Get(), cnv.Length());
 			//return (TCHAR*)FUTF8ToTCHAR((const ANSICHAR*)InArray.GetData(), InArray.Num()).Get();
 			//FString UTF8Text(InArray.Num(), (TCHAR*)FUTF8ToTCHAR((const ANSICHAR*)InArray.GetData(), InArray.Num()).Get());
 			return MoveTemp(UTF8Text);
diff --git a/Engine/Plugins/Media/PixelStreaming/Source/PixelStreamingServers/Private/ServerUtils.cpp b/Engine/Plugins/Media/PixelStreaming/Source/PixelStreamingServers/Private/ServerUtils.cpp
--- a/Engine/Plugins/Media/PixelStreaming/Source/PixelStreamingServers/Private/ServerUtils.cpp
+++ b/Engine/Plugins/Media/PixelStreaming/Source/PixelStreamingServers/Private/ServerUtils.cpp
@@ -266,7 +266,7 @@ namespace UE::PixelStreamingServers::Utils
 	FString ToString(TArrayView<uint8> UTF8Bytes)
 	{
 		FUTF8ToTCHAR Converted((const ANSICHAR*)UTF8Bytes.GetData(), UTF8Bytes.Num());
-		FString OutString(Converted.Length(), Converted.Get());
+		FString OutString = FString::ConstructFromPtrSize(Converted.Get(), Converted.Length());
 		return OutString;
 	}
 
diff --git a/Engine/Plugins/Online/OnlineServices/Source/OnlineServicesCommon/Private/Online/OnlineConfig.cpp b/Engine/Plugins/Online/OnlineServices/Source/OnlineServicesCommon/Private/Online/OnlineConfig.cpp
--- a/Engine/Plugins/Online/OnlineServices/Source/OnlineServicesCommon/Private/Online/OnlineConfig.cpp
+++ b/Engine/Plugins/Online/OnlineServices/Source/OnlineServicesCommon/Private/Online/OnlineConfig.cpp
@@ -28,7 +28,7 @@ const TCHAR* ParseField(const TCHAR* InStr, TMap<FString, FString>& OutMap)
 	if (*Str)
 	{
 		int32 Len = UE_PTRDIFF_TO_INT32(Str - Start);
-		FString FieldName(Len, Start);
+		FString FieldName = FString::ConstructFromPtrSize(Start, Len);
 
 		Str++;
 		Start = Str;
@@ -44,7 +44,7 @@ const TCHAR* ParseField(const TCHAR* InStr, TMap<FString, FString>& OutMap)
 			}
 
 			Len = UE_PTRDIFF_TO_INT32(Str - Start);
-			FString FieldValue(Len, Start);
+			FString FieldValue= FString::ConstructFromPtrSize(Start, Len);
 
 			OutMap.Emplace(MoveTemp(FieldName), MoveTemp(FieldValue));
 			Str++;
@@ -84,7 +84,7 @@ const TCHAR* ParseField(const TCHAR* InStr, TMap<FString, FString>& OutMap)
 			Str++;
 
 			Len = UE_PTRDIFF_TO_INT32(Str - Start);
-			FString FieldValue(Len, Start);
+			FString FieldValue = FString::ConstructFromPtrSize(Start, Len);
 
 			OutMap.Emplace(MoveTemp(FieldName), MoveTemp(FieldValue));
 		}
@@ -94,7 +94,7 @@ const TCHAR* ParseField(const TCHAR* InStr, TMap<FString, FString>& OutMap)
 			Str = FAsciiSet::Skip(Str, StringChars);
 				
 			Len = UE_PTRDIFF_TO_INT32(Str - Start);
-			FString FieldValue(Len, Start);
+			FString FieldValue = FString::ConstructFromPtrSize(Start, Len);
 
 			OutMap.Emplace(MoveTemp(FieldName), MoveTemp(FieldValue));
 		}
@@ -217,7 +217,7 @@ TArray<FString> FOnlineConfigStructGConfig::CreateValueArray(const FString& InCo
 					}
 
 					const int32 Len = UE_PTRDIFF_TO_INT32(Str - Start);
-					FString FieldValue(Len, Start);
+					FString FieldValue = FString::ConstructFromPtrSize(Start, Len);
 					Result.Emplace(MoveTemp(FieldValue));
 					Str++;
 				}
@@ -227,7 +227,7 @@ TArray<FString> FOnlineConfigStructGConfig::CreateValueArray(const FString& InCo
 					Str = FAsciiSet::Skip(Str, StringChars);
 
 					const int32 Len = UE_PTRDIFF_TO_INT32(Str - Start);
-					FString FieldValue(Len, Start);
+					FString FieldValue = FString::ConstructFromPtrSize(Start, Len);
 					Result.Emplace(MoveTemp(FieldValue));
 				}
 
diff --git a/Engine/Plugins/Online/OnlineServices/Source/OnlineServicesCommon/Public/Online/SessionsCommon.h b/Engine/Plugins/Online/OnlineServices/Source/OnlineServicesCommon/Public/Online/SessionsCommon.h
--- a/Engine/Plugins/Online/OnlineServices/Source/OnlineServicesCommon/Public/Online/SessionsCommon.h
+++ b/Engine/Plugins/Online/OnlineServices/Source/OnlineServicesCommon/Public/Online/SessionsCommon.h
@@ -49,7 +49,7 @@ public:
 	virtual inline FOnlineSessionId FromReplicationData(const TArray<uint8>& ReplicationData) override
 	{
 		const FUTF8ToTCHAR IdValueTCHAR((char*)ReplicationData.GetData(), ReplicationData.Num());
-		const FString IdValue = FString(IdValueTCHAR.Length(), IdValueTCHAR.Get());
+		const FString IdValue = FString::ConstructFromPtrSize(IdValueTCHAR.Get(), IdValueTCHAR.Length());
 
 		if (!IdValue.IsEmpty())
 		{
@@ -114,7 +114,7 @@ public:
 	virtual inline FSessionInviteId FromReplicationData(const TArray<uint8>& ReplicationData) override
 	{
 		const FUTF8ToTCHAR IdValueTCHAR((char*)ReplicationData.GetData(), ReplicationData.Num());
-		const FString IdValue = FString(IdValueTCHAR.Length(), IdValueTCHAR.Get());
+		const FString IdValue = FString::ConstructFromPtrSize(IdValueTCHAR.Get(), IdValueTCHAR.Length());
 
 		if (!IdValue.IsEmpty())
 		{
diff --git a/Engine/Plugins/Online/OnlineServices/Source/OnlineServicesInterface/Private/Online/SchemaTypes.cpp b/Engine/Plugins/Online/OnlineServices/Source/OnlineServicesInterface/Private/Online/SchemaTypes.cpp
--- a/Engine/Plugins/Online/OnlineServices/Source/OnlineServicesInterface/Private/Online/SchemaTypes.cpp
+++ b/Engine/Plugins/Online/OnlineServices/Source/OnlineServicesInterface/Private/Online/SchemaTypes.cpp
@@ -398,7 +398,7 @@ void LexFromString(FSchemaVariant& OutSchemaVariant, const TCHAR* InStr)
 	if (const TCHAR* ValueStr = FCString::Strchr(InStr, ':'))
 	{
 		const int32 TypeLen = UE_PTRDIFF_TO_INT32(ValueStr - InStr);
-		const FString TypeStr(TypeLen, InStr);
+		const FString TypeStr = FString::ConstructFromPtrSize(InStr, TypeLen);
 		ValueStr++;
 
 		LexFromString(OutSchemaVariant.VariantType, *TypeStr);
diff --git a/Engine/Plugins/Online/OnlineSubsystemAmazon/Source/Private/OnlineIdentityAmazon.cpp b/Engine/Plugins/Online/OnlineSubsystemAmazon/Source/Private/OnlineIdentityAmazon.cpp
--- a/Engine/Plugins/Online/OnlineSubsystemAmazon/Source/Private/OnlineIdentityAmazon.cpp
+++ b/Engine/Plugins/Online/OnlineSubsystemAmazon/Source/Private/OnlineIdentityAmazon.cpp
@@ -256,7 +256,7 @@ FUniqueNetIdPtr FOnlineIdentityAmazon::CreateUniquePlayerId(uint8* Bytes, int32
 {
 	if (Bytes != NULL && Size > 0)
 	{
-		FString StrId(Size, (TCHAR*)Bytes);
+		FString StrId = FString::ConstructFromPtrSize((TCHAR*)Bytes, Size);
 		return FUniqueNetIdAmazon::Create(StrId);
 	}
 	return NULL;
diff --git a/Engine/Plugins/Online/OnlineSubsystemFacebook/Source/Private/OnlineIdentityFacebookCommon.cpp b/Engine/Plugins/Online/OnlineSubsystemFacebook/Source/Private/OnlineIdentityFacebookCommon.cpp
--- a/Engine/Plugins/Online/OnlineSubsystemFacebook/Source/Private/OnlineIdentityFacebookCommon.cpp
+++ b/Engine/Plugins/Online/OnlineSubsystemFacebook/Source/Private/OnlineIdentityFacebookCommon.cpp
@@ -204,7 +204,7 @@ FUniqueNetIdPtr FOnlineIdentityFacebookCommon::CreateUniquePlayerId(uint8* Bytes
 {
 	if (Bytes != nullptr && Size > 0)
 	{
-		FString StrId(Size, (TCHAR*)Bytes);
+		FString StrId = FString::ConstructFromPtrSize((TCHAR*)Bytes, Size);
 		return FUniqueNetIdFacebook::Create(StrId);
 	}
 	return nullptr;
diff --git a/Engine/Plugins/Online/OnlineSubsystemNull/Source/Private/OnlineIdentityNull.cpp b/Engine/Plugins/Online/OnlineSubsystemNull/Source/Private/OnlineIdentityNull.cpp
--- a/Engine/Plugins/Online/OnlineSubsystemNull/Source/Private/OnlineIdentityNull.cpp
+++ b/Engine/Plugins/Online/OnlineSubsystemNull/Source/Private/OnlineIdentityNull.cpp
@@ -251,7 +251,7 @@ FUniqueNetIdPtr FOnlineIdentityNull::CreateUniquePlayerId(uint8* Bytes, int32 Si
 {
 	if (Bytes != NULL && Size > 0)
 	{
-		FString StrId(Size, (TCHAR*)Bytes);
+		FString StrId = FString::ConstructFromPtrSize((TCHAR*)Bytes, Size);
 		return FUniqueNetIdNull::Create(StrId);
 	}
 	return NULL;
diff --git a/Engine/Plugins/Runtime/OpenXR/Source/OpenXRHMD/Private/OpenXRCore.cpp b/Engine/Plugins/Runtime/OpenXR/Source/OpenXRHMD/Private/OpenXRCore.cpp
--- a/Engine/Plugins/Runtime/OpenXR/Source/OpenXRHMD/Private/OpenXRCore.cpp
+++ b/Engine/Plugins/Runtime/OpenXR/Source/OpenXRHMD/Private/OpenXRCore.cpp
@@ -42,7 +42,7 @@ FString FOpenXRPath::ToString() const
 	char PathChars[XR_MAX_PATH_LENGTH];
 	XrResult Result = xrPathToString(Instance, Path, XR_MAX_PATH_LENGTH, &PathCount, PathChars);
 	check(XR_SUCCEEDED(Result));
-	return XR_SUCCEEDED(Result) ? FString(PathCount - 1, PathChars) : FString();
+	return XR_SUCCEEDED(Result) ? FString::ConstructFromPtrSize(PathChars, PathCount - 1) : FString();
 }
 
 uint32 FOpenXRPath::GetStringLength() const
diff --git a/Engine/Plugins/Runtime/OpenXR/Source/OpenXRHMD/Private/OpenXRHMDModule.cpp b/Engine/Plugins/Runtime/OpenXR/Source/OpenXRHMD/Private/OpenXRHMDModule.cpp
--- a/Engine/Plugins/Runtime/OpenXR/Source/OpenXRHMD/Private/OpenXRHMDModule.cpp
+++ b/Engine/Plugins/Runtime/OpenXR/Source/OpenXRHMD/Private/OpenXRHMDModule.cpp
@@ -118,7 +118,7 @@ bool FOpenXRHMDModule::PreInit()
 			{
 				WCHAR DeviceGuid[XR_MAX_AUDIO_DEVICE_STR_SIZE_OCULUS];
 				GetAudioOutputDeviceGuidOculus(Instance, DeviceGuid);
-				OculusAudioOutputDevice = FString(XR_MAX_AUDIO_DEVICE_STR_SIZE_OCULUS, DeviceGuid);
+				OculusAudioOutputDevice = FString::ConstructFromPtrSize(DeviceGuid, XR_MAX_AUDIO_DEVICE_STR_SIZE_OCULUS);
 			}
 		}
 		if (!CVarRetainPreInitInstance.GetValueOnAnyThread())
diff --git a/Engine/Plugins/Runtime/PropertyAccess/Source/PropertyAccessEditor/Private/PropertyAccessEditor.cpp b/Engine/Plugins/Runtime/PropertyAccess/Source/PropertyAccessEditor/Private/PropertyAccessEditor.cpp
--- a/Engine/Plugins/Runtime/PropertyAccess/Source/PropertyAccessEditor/Private/PropertyAccessEditor.cpp
+++ b/Engine/Plugins/Runtime/PropertyAccess/Source/PropertyAccessEditor/Private/PropertyAccessEditor.cpp
@@ -214,7 +214,7 @@ struct FPropertyAccessEditorSystem
 				int32 ArrayIndex = INDEX_NONE;
 				PropertyPathHelpers::FindFieldNameAndArrayIndex(SegmentString.Len(), *SegmentString, PropertyNameLength, &PropertyNamePtr, ArrayIndex);
 				ensure(PropertyNamePtr != nullptr);
-				FString PropertyNameString(PropertyNameLength, PropertyNamePtr);
+				FString PropertyNameString = FString::ConstructFromPtrSize(PropertyNamePtr, PropertyNameLength);
 				Segment.Name = FName(*PropertyNameString, FNAME_Find);
 				Segment.ArrayIndex = ArrayIndex;
 
diff --git a/Engine/Source/Developer/ShaderCompilerCommon/Private/HlslLexer.cpp b/Engine/Source/Developer/ShaderCompilerCommon/Private/HlslLexer.cpp
--- a/Engine/Source/Developer/ShaderCompilerCommon/Private/HlslLexer.cpp
+++ b/Engine/Source/Developer/ShaderCompilerCommon/Private/HlslLexer.cpp
@@ -752,7 +752,7 @@ namespace CrossCompiler
 			}
 
 		Done:
-			OutLiteral = FString(static_cast<int32>(Current - Original), Original);
+			OutLiteral = FString::ConstructFromPtrSize(Original, static_cast<int32>(Current - Original));
 			OutType = Type;
 			return true;
 
diff --git a/Engine/Source/Developer/ShaderPreprocessor/Private/ShaderPreprocessor.cpp b/Engine/Source/Developer/ShaderPreprocessor/Private/ShaderPreprocessor.cpp
--- a/Engine/Source/Developer/ShaderPreprocessor/Private/ShaderPreprocessor.cpp
+++ b/Engine/Source/Developer/ShaderPreprocessor/Private/ShaderPreprocessor.cpp
@@ -387,7 +387,7 @@ static const ANSICHAR* StbResolveInclude(const ANSICHAR* PathInSource, uint32 Pa
 	}
 
 	// Slow path...  Platform specific files and procedurally generated files (/Engine/Generated/Material.ush) -- typically 5% of files.
-	FString PathModified(PathLen, PathInSource);
+	FString PathModified = FString::ConstructFromPtrSize(PathInSource, PathLen);
 	if (!PathModified.StartsWith(TEXT("/"))) // if path doesn't start with / it's relative, if so append the parent's folder and collapse any relative dirs
 	{
 		FString ParentFolder(ParentPathAnsi);
diff --git a/Engine/Source/Developer/TraceAnalysis/Private/Analysis/Engine.cpp b/Engine/Source/Developer/TraceAnalysis/Private/Analysis/Engine.cpp
--- a/Engine/Source/Developer/TraceAnalysis/Private/Analysis/Engine.cpp
+++ b/Engine/Source/Developer/TraceAnalysis/Private/Analysis/Engine.cpp
@@ -948,7 +948,7 @@ bool IAnalyzer::FEventData::GetString(const ANSICHAR* FieldName, FString& Out) c
 
 		if (Field.SizeAndType == sizeof(ANSICHAR))
 		{
-			Out = FString(Data->DataSize, (const ANSICHAR*)(Data->Data));
+			Out = FString::ConstructFromPtrSize((const ANSICHAR*)(Data->Data), Data->DataSize);
 			return true;
 		}
 
diff --git a/Engine/Source/Developer/TraceInsights/Private/Insights/ViewModels/ThreadTimingTrack.cpp b/Engine/Source/Developer/TraceInsights/Private/Insights/ViewModels/ThreadTimingTrack.cpp
--- a/Engine/Source/Developer/TraceInsights/Private/Insights/ViewModels/ThreadTimingTrack.cpp
+++ b/Engine/Source/Developer/TraceInsights/Private/Insights/ViewModels/ThreadTimingTrack.cpp
@@ -102,7 +102,7 @@ static void AppendMetadataToTooltip(FTooltipDrawState& Tooltip, TArrayView<const
 			break;
 		}
 
-		FString Key(static_cast<int32>(Context.AsLength()), Context.AsCString());
+		FString Key = FString::ConstructFromPtrSize(Context.AsCString(), static_cast<int32>(Context.AsLength()));
 		Key += TEXT(":");
 
 		if (!CborReader.ReadNext(Context))
diff --git a/Engine/Source/Developer/TraceInsights/Private/Insights/ViewModels/TimerFilters.cpp b/Engine/Source/Developer/TraceInsights/Private/Insights/ViewModels/TimerFilters.cpp
--- a/Engine/Source/Developer/TraceInsights/Private/Insights/ViewModels/TimerFilters.cpp
+++ b/Engine/Source/Developer/TraceInsights/Private/Insights/ViewModels/TimerFilters.cpp
@@ -342,7 +342,7 @@ bool FMetadataFilterState::ApplyFilterToMetadata(TArrayView<const uint8>& Metada
 			break;
 		}
 
-		FString CurrentKey(static_cast<int32>(Context.AsLength()), Context.AsCString());
+		FString CurrentKey = FString::ConstructFromPtrSize(Context.AsCString(), static_cast<int32>(Context.AsLength()));
 
 		if (!CborReader.ReadNext(Context))
 		{
diff --git a/Engine/Source/Developer/TraceInsights/Private/Insights/Widgets/SStartPageWindow.h b/Engine/Source/Developer/TraceInsights/Private/Insights/Widgets/SStartPageWindow.h
--- a/Engine/Source/Developer/TraceInsights/Private/Insights/Widgets/SStartPageWindow.h
+++ b/Engine/Source/Developer/TraceInsights/Private/Insights/Widgets/SStartPageWindow.h
@@ -113,7 +113,7 @@ struct FTraceViewModel
 
 	static FText AnsiStringViewToText(const FAnsiStringView& AnsiStringView)
 	{
-		FString FatString(AnsiStringView.Len(), AnsiStringView.GetData());
+		FString FatString = FString::ConstructFromPtrSize(AnsiStringView.GetData(), AnsiStringView.Len());
 		return FText::FromString(FatString);
 	}
 };
diff --git a/Engine/Source/Developer/TraceServices/Private/Analyzers/DiagnosticsAnalysis.cpp b/Engine/Source/Developer/TraceServices/Private/Analyzers/DiagnosticsAnalysis.cpp
--- a/Engine/Source/Developer/TraceServices/Private/Analyzers/DiagnosticsAnalysis.cpp
+++ b/Engine/Source/Developer/TraceServices/Private/Analyzers/DiagnosticsAnalysis.cpp
@@ -47,15 +47,15 @@ bool FDiagnosticsAnalyzer::OnEvent(uint16 RouteId, EStyle Style, const FOnEventC
 		uint8 AppNameOffset = EventData.GetValue<uint8>("AppNameOffset");
 		uint8 CommandLineOffset = EventData.GetValue<uint8>("CommandLineOffset");
 
-		SessionInfo.Platform = FString(AppNameOffset, (const ANSICHAR*)Attachment);
+		SessionInfo.Platform = FString::ConstructFromPtrSize((const ANSICHAR*)Attachment, AppNameOffset);
 
 		Attachment += AppNameOffset;
 		int32 AppNameLength = CommandLineOffset - AppNameOffset;
-		SessionInfo.AppName = FString(AppNameLength, (const ANSICHAR*)Attachment);
+		SessionInfo.AppName = FString::ConstructFromPtrSize((const ANSICHAR*)Attachment, AppNameLength);
 
 		Attachment += AppNameLength;
 		int32 CommandLineLength = EventData.GetAttachmentSize() - CommandLineOffset;
-		SessionInfo.CommandLine = FString(CommandLineLength, (const ANSICHAR*)Attachment);
+		SessionInfo.CommandLine = FString::ConstructFromPtrSize((const ANSICHAR*)Attachment, CommandLineLength);
 
 		SessionInfo.ConfigurationType = (EBuildConfiguration)EventData.GetValue<uint8>("ConfigurationType");
 
diff --git a/Engine/Source/Developer/TraceServices/Public/TraceServices/Utils.h b/Engine/Source/Developer/TraceServices/Public/TraceServices/Utils.h
--- a/Engine/Source/Developer/TraceServices/Public/TraceServices/Utils.h
+++ b/Engine/Source/Developer/TraceServices/Public/TraceServices/Utils.h
@@ -17,9 +17,9 @@ inline FString LegacyAttachmentString(
 	FString Out;
 	if (!Context.EventData.GetString(FieldName, Out))
 	{
-		Out = FString(
-				Context.EventData.GetAttachmentSize() / sizeof(AttachedCharType),
-				(const AttachedCharType*)(Context.EventData.GetAttachment()));
+		Out = FString::ConstructFromPtrSize(
+				(const AttachedCharType*)(Context.EventData.GetAttachment()),
+				Context.EventData.GetAttachmentSize() / sizeof(AttachedCharType));
 	}
 	return Out;
 }
diff --git a/Engine/Source/Developer/UbaCoordinatorHorde/Private/UbaHordeAgent.cpp b/Engine/Source/Developer/UbaCoordinatorHorde/Private/UbaHordeAgent.cpp
--- a/Engine/Source/Developer/UbaCoordinatorHorde/Private/UbaHordeAgent.cpp
+++ b/Engine/Source/Developer/UbaCoordinatorHorde/Private/UbaHordeAgent.cpp
@@ -94,7 +94,7 @@ bool FUbaHordeAgent::UploadBinaries(const FString& BundleDirectory, const char*
 
 	auto FindOrAddBlobFile = [&BlobFiles, &BundleDirectory](const AgentMessage::FBlobRequest& BlobRequest) -> FArchive*
 		{
-			const FString Locator(BlobRequest.Locator.Len(), BlobRequest.Locator.GetData());
+			const FString Locator = FString::ConstructFromPtrSize(BlobRequest.Locator.GetData(), BlobRequest.Locator.Len());
 
 			if (Locator.IsEmpty())
 			{
diff --git a/Engine/Source/Developer/VulkanShaderFormat/Private/VulkanShaderCompiler.cpp b/Engine/Source/Developer/VulkanShaderFormat/Private/VulkanShaderCompiler.cpp
--- a/Engine/Source/Developer/VulkanShaderFormat/Private/VulkanShaderCompiler.cpp
+++ b/Engine/Source/Developer/VulkanShaderFormat/Private/VulkanShaderCompiler.cpp
@@ -1256,7 +1256,7 @@ static FString ParseEntrypointDecl(FShaderSource::FViewType PreprocessedShader,
 		}
 		++DeclBegin;
 
-		EntryPointDecl = FString(DeclEnd - DeclBegin, &PreprocessedShader[DeclBegin]);
+		EntryPointDecl = FString::ConstructFromPtrSize(&PreprocessedShader[DeclBegin], DeclEnd - DeclBegin);
 		EraseDebugLines(EntryPointDecl);
 		EntryPointDecl.TrimStartAndEndInline();
 		break;
@@ -1283,7 +1283,7 @@ uint8 ParseWaveSize(
 			const int32 EndNumber = EntrypointDecl.Find(TEXT(")"), ESearchCase::CaseSensitive, ESearchDir::FromStart, StartNumber);
 			check(EndNumber != INDEX_NONE);
 
-			FString WaveSizeValue(EndNumber - StartNumber, &EntrypointDecl[StartNumber]);
+			FString WaveSizeValue = FString::ConstructFromPtrSize(&EntrypointDecl[StartNumber], EndNumber - StartNumber);
 			WaveSizeValue.RemoveSpacesInline();
 			if (WaveSizeValue != TEXT("N"))  // skip the macro decl
 			{
@@ -1503,7 +1503,7 @@ static TArray<FString> ConvertUBToBindless(FString& PreprocessedShaderSource)
 			const TCHAR* MemberTypeStartPtr = nullptr;
 			const TCHAR* MemberTypeEndPtr = nullptr;
 			ParseHLSLTypeName(MemberSearchPtr, MemberTypeStartPtr, MemberTypeEndPtr);
-			const FString MemberTypeName(MemberTypeEndPtr - MemberTypeStartPtr, MemberTypeStartPtr);
+			const FString MemberTypeName = FString::ConstructFromPtrSize(MemberTypeStartPtr, MemberTypeEndPtr - MemberTypeStartPtr);
 
 			FString MemberName;
 			MemberSearchPtr = ParseHLSLSymbolName(MemberTypeEndPtr, MemberName);
@@ -1570,7 +1570,7 @@ static TArray<FString> ConvertUBToBindless(FString& PreprocessedShaderSource)
 			check(ClosingBracePtr);
 			const int32 ClosingBraceIndex = ClosingBracePtr - (*PreprocessedShaderSource);
 
-			const FString Members(ClosingBracePtr - OpeningBracePtr - 1, OpeningBracePtr + 1);
+			const FString Members = FString::ConstructFromPtrSize(OpeningBracePtr + 1, ClosingBracePtr - OpeningBracePtr - 1);
 			const FString NewDecl = GenerateNewDecl(CBIndex, Members, StructName);
 
 			const int32 OldDeclLen = ClosingBraceIndex - SearchIndex + 1;
diff --git a/Engine/Source/Developer/Windows/LiveCoding/Private/LiveCodingModule.cpp b/Engine/Source/Developer/Windows/LiveCoding/Private/LiveCodingModule.cpp
--- a/Engine/Source/Developer/Windows/LiveCoding/Private/LiveCodingModule.cpp
+++ b/Engine/Source/Developer/Windows/LiveCoding/Private/LiveCodingModule.cpp
@@ -349,7 +349,7 @@ void FLiveCodingModule::StartupModule()
 	{
 		TCHAR Scratch[MAX_PATH];
 		int Length = GetModuleFileNameEx(::GetCurrentProcess(), NULL, Scratch, MAX_PATH);
-		MainModuleName = FString(Length, Scratch);
+		MainModuleName = FString::ConstructFromPtrSize(Scratch, Length);
 	}
 
 	// https://docs.microsoft.com/en-us/windows/win32/api/winternl/ns-winternl-peb_ldr_data
@@ -401,7 +401,7 @@ void FLiveCodingModule::StartupModule()
 			break;
 		}
 
-		FString FullPath(ModuleData.FullDllName.Length / sizeof(ModuleData.FullDllName.Buffer[0]), ModuleData.FullDllName.Buffer);
+		FString FullPath = FString::ConstructFromPtrSize(ModuleData.FullDllName.Buffer, ModuleData.FullDllName.Length / sizeof(ModuleData.FullDllName.Buffer[0]));
 		if (!FullPath.Equals(MainModuleName, ESearchCase::IgnoreCase))
 		{
 			FPaths::NormalizeFilename(FullPath);
@@ -1496,7 +1496,7 @@ void FLiveCodingModule::OnDllNotification(unsigned int Reason, const void* DataP
 		UPTRINT	Base;
 	};
 	const auto& Data = *(FNotificationData*)DataPtr;
-	FString FullPath(Data.FullPath.Length / sizeof(Data.FullPath.Buffer[0]), Data.FullPath.Buffer);
+	FString FullPath = FString::ConstructFromPtrSize(Data.FullPath.Buffer, Data.FullPath.Length / sizeof(Data.FullPath.Buffer[0]));
 	FPaths::NormalizeFilename(FullPath);
 
 	switch (Reason)
diff --git a/Engine/Source/Editor/UnrealEd/Private/Commandlets/GatherTextFromSourceCommandlet.cpp b/Engine/Source/Editor/UnrealEd/Private/Commandlets/GatherTextFromSourceCommandlet.cpp
--- a/Engine/Source/Editor/UnrealEd/Private/Commandlets/GatherTextFromSourceCommandlet.cpp
+++ b/Engine/Source/Editor/UnrealEd/Private/Commandlets/GatherTextFromSourceCommandlet.cpp
@@ -1115,7 +1115,7 @@ bool UGatherTextFromSourceCommandlet::FMacroArgumentGatherer::Gather(const TCHAR
 	{
 		return false;
 	}
-	FString NewArgument = FString(Count, Arg);
+	FString NewArgument = FString::ConstructFromPtrSize(Arg, Count);
 	NewArgument.TrimEndInline();
 
 	int32 CurrentArgLen = CurrentArgument.Len();
diff --git a/Engine/Source/Editor/UnrealEd/Private/EditorCategoryUtils.cpp b/Engine/Source/Editor/UnrealEd/Private/EditorCategoryUtils.cpp
--- a/Engine/Source/Editor/UnrealEd/Private/EditorCategoryUtils.cpp
+++ b/Engine/Source/Editor/UnrealEd/Private/EditorCategoryUtils.cpp
@@ -271,13 +271,13 @@ FString FEditorCategoryUtils::GetCategoryDisplayString(const FString& Unsanitize
 			int32 EndIndex = DisplayString.Find(TEXT("}"), ESearchCase::CaseSensitive, ESearchDir::FromStart, KeyIndex);
 			if (EndIndex != INDEX_NONE)
 			{
-				FString ToReplaceStr(EndIndex+1 - KeyIndex, *DisplayString + KeyIndex);
+				FString ToReplaceStr = FString::ConstructFromPtrSize(*DisplayString + KeyIndex, EndIndex+1 - KeyIndex);
 				FString ReplacementStr;
 				
 				int32 KeyLen = EndIndex - (KeyIndex + 1);
 				if (KeyLen > 0)
 				{
-					FString Key(KeyLen, *DisplayString + KeyIndex+1);
+					FString Key = FString::ConstructFromPtrSize(*DisplayString + KeyIndex+1, KeyLen);
 					Key.TrimStartInline();
 					ReplacementStr = FEditorCategoryUtilsImpl::GetCategory(*Key).ToString();
 				}
diff --git a/Engine/Source/Editor/UnrealEd/Public/Kismet2/CompilerResultsLog.h b/Engine/Source/Editor/UnrealEd/Public/Kismet2/CompilerResultsLog.h
--- a/Engine/Source/Editor/UnrealEd/Public/Kismet2/CompilerResultsLog.h
+++ b/Engine/Source/Editor/UnrealEd/Public/Kismet2/CompilerResultsLog.h
@@ -357,7 +357,7 @@ protected:
 		if (const TCHAR* DelimiterStr = FCString::Strstr(Format, TEXT("@@")))
 		{
 			int32 TokenLength = UE_PTRDIFF_TO_INT32(DelimiterStr - Format);
-			OutMessage.AddToken(FTextToken::Create(FText::FromString(FString(TokenLength, Format))));
+			OutMessage.AddToken(FTextToken::Create(FText::FromString(FString::ConstructFromPtrSize(Format, TokenLength))));
 			FEdGraphToken::Create(First, this, OutMessage, OutSourceNode);
 
 			const TCHAR* NextChunk = DelimiterStr + FCString::Strlen(TEXT("@@"));
diff --git a/Engine/Source/Runtime/Core/Private/Internationalization/ICUTextCharacterIterator.cpp b/Engine/Source/Runtime/Core/Private/Internationalization/ICUTextCharacterIterator.cpp
--- a/Engine/Source/Runtime/Core/Private/Internationalization/ICUTextCharacterIterator.cpp
+++ b/Engine/Source/Runtime/Core/Private/Internationalization/ICUTextCharacterIterator.cpp
@@ -24,7 +24,7 @@ FICUTextCharacterIterator_NativeUTF16::FICUTextCharacterIterator_NativeUTF16(FSt
 
 FICUTextCharacterIterator_NativeUTF16::FICUTextCharacterIterator_NativeUTF16(const FICUTextCharacterIterator_NativeUTF16& Other)
 	: icu::UCharCharacterIterator(Other)
-	, InternalString(FString(Other.StringRef.Len(), Other.StringRef.GetData()))
+	, InternalString(FString::ConstructFromPtrSize(Other.StringRef.GetData(), Other.StringRef.Len()))
 	, StringRef(InternalString)
 {
 	SetTextFromStringRef();
@@ -78,7 +78,7 @@ FICUTextCharacterIterator_ConvertToUnicodeString::FICUTextCharacterIterator_Conv
 }
 
 FICUTextCharacterIterator_ConvertToUnicodeString::FICUTextCharacterIterator_ConvertToUnicodeString(FStringView InString)
-	: FICUTextCharacterIterator_ConvertToUnicodeStringPrivate(FString(InString.Len(), InString.GetData()))
+	: FICUTextCharacterIterator_ConvertToUnicodeStringPrivate(FString::ConstructFromPtrSize(InString.GetData(), InString.Len()))
 	, icu::StringCharacterIterator(InternalString)
 {
 }
diff --git a/Engine/Source/Runtime/CoreUObject/Private/Misc/PackagePath.cpp b/Engine/Source/Runtime/CoreUObject/Private/Misc/PackagePath.cpp
--- a/Engine/Source/Runtime/CoreUObject/Private/Misc/PackagePath.cpp
+++ b/Engine/Source/Runtime/CoreUObject/Private/Misc/PackagePath.cpp
@@ -661,7 +661,7 @@ FArchive& operator<<(FArchive& Ar, FPackagePath& PackagePath)
 	{
 		if (StringDataLen > 0)
 		{
-			FString SerializerString = FString(StringDataLen, PackagePath.StringData.Get());
+			FString SerializerString = FString::ConstructFromPtrSize(PackagePath.StringData.Get(), StringDataLen);
 			Ar << SerializerString;
 		}
 	}
diff --git a/Engine/Source/Runtime/Engine/Private/AssetManager.cpp b/Engine/Source/Runtime/Engine/Private/AssetManager.cpp
--- a/Engine/Source/Runtime/Engine/Private/AssetManager.cpp
+++ b/Engine/Source/Runtime/Engine/Private/AssetManager.cpp
@@ -1841,7 +1841,7 @@ TSharedPtr<FStreamableHandle> UAssetManager::ChangeBundleStateForPrimaryAssets(c
 				DebugName << TEXT(")");
 			}
 
-			NewHandle = LoadAssetList(MoveTemp(PathsToLoad), FStreamableDelegate(), Priority, FString(DebugName.Len(), DebugName.ToString()));
+			NewHandle = LoadAssetList(MoveTemp(PathsToLoad), FStreamableDelegate(), Priority, FString::ConstructFromPtrSize(DebugName.ToString(), DebugName.Len()));
 
 			if (!NewHandle.IsValid())
 			{
diff --git a/Engine/Source/Runtime/Engine/Private/EdGraph/EdGraphPin.cpp b/Engine/Source/Runtime/Engine/Private/EdGraph/EdGraphPin.cpp
--- a/Engine/Source/Runtime/Engine/Private/EdGraph/EdGraphPin.cpp
+++ b/Engine/Source/Runtime/Engine/Private/EdGraph/EdGraphPin.cpp
@@ -1174,7 +1174,7 @@ bool UEdGraphPin::ImportTextItem(const TCHAR*& Buffer, int32 PortFlags, class UO
 		// Advance over the '='
 		Buffer++;
 
-		FString PropertyToken(NumCharsInToken, StartBuffer);
+		FString PropertyToken = FString::ConstructFromPtrSize(StartBuffer, NumCharsInToken);
 		PropertyToken.TrimEndInline();
 		bool bParseSuccess = false;
 		if (PropertyToken == PinHelpers::PinIdName)
@@ -2166,7 +2166,7 @@ EPinResolveResult PinHelpers::ImportText_PinReference(const TCHAR*& Buffer, UEdG
 		Buffer++;
 	}
 
-	FString PinRefLine(Buffer - BufferStart, BufferStart);
+	FString PinRefLine = FString::ConstructFromPtrSize(BufferStart, Buffer - BufferStart);
 	FString OwningNodeString;
 	FString PinGuidString;
 	if (PinRefLine.Split(TEXT(" "), &OwningNodeString, &PinGuidString, ESearchCase::CaseSensitive, ESearchDir::FromEnd))
diff --git a/Engine/Source/Runtime/Engine/Private/KismetStringLibrary.cpp b/Engine/Source/Runtime/Engine/Private/KismetStringLibrary.cpp
--- a/Engine/Source/Runtime/Engine/Private/KismetStringLibrary.cpp
+++ b/Engine/Source/Runtime/Engine/Private/KismetStringLibrary.cpp
@@ -420,7 +420,7 @@ TArray<FString> UKismetStringLibrary::GetCharacterArrayFromString(const FString&
 		for (auto CharIt(SourceString.CreateConstIterator()); CharIt; ++CharIt)
 		{
 			TCHAR Char = *CharIt;
-			SeparatedChars.Add(FString(1, &Char));
+			SeparatedChars.Add(FString::ConstructFromPtrSize(&Char, 1));
 		}
 
 		// Remove the null terminator on the end
diff --git a/Engine/Source/Runtime/Experimental/IoStore/OnDemand/Private/OnDemandIoStore.cpp b/Engine/Source/Runtime/Experimental/IoStore/OnDemand/Private/OnDemandIoStore.cpp
--- a/Engine/Source/Runtime/Experimental/IoStore/OnDemand/Private/OnDemandIoStore.cpp
+++ b/Engine/Source/Runtime/Experimental/IoStore/OnDemand/Private/OnDemandIoStore.cpp
@@ -357,7 +357,7 @@ FString LexToString(EOnDemandContainerFlags Flags)
 {
 	TStringBuilder<128> Sb;
 	LexToString(Flags, Sb);
-	return FString(Sb.ToString(), Sb.Len());
+	return FString::ConstructFromPtrSize(Sb.ToString(), Sb.Len());
 }
 
 ///////////////////////////////////////////////////////////////////////////////
diff --git a/Engine/Source/Runtime/Experimental/Iris/Core/Private/Iris/Serialization/NetBitStreamUtil.cpp b/Engine/Source/Runtime/Experimental/Iris/Core/Private/Iris/Serialization/NetBitStreamUtil.cpp
--- a/Engine/Source/Runtime/Experimental/Iris/Core/Private/Iris/Serialization/NetBitStreamUtil.cpp
+++ b/Engine/Source/Runtime/Experimental/Iris/Core/Private/Iris/Serialization/NetBitStreamUtil.cpp
@@ -286,7 +286,7 @@ void ReadString(FNetBitStreamReader* Reader, FString& OutString)
 		}
 
 		// There are two FString constructors with similar signature. The one with known length of the string passes the count first.
-		OutString = FString(static_cast<int32>(Length), CharBuffer.GetData());
+		OutString = FString::ConstructFromPtrSize(CharBuffer.GetData(), static_cast<int32>(Length));
 	}
 }
 
diff --git a/Engine/Source/Runtime/Experimental/Iris/Core/Private/Iris/Serialization/StringNetSerializerUtils.cpp b/Engine/Source/Runtime/Experimental/Iris/Core/Private/Iris/Serialization/StringNetSerializerUtils.cpp
--- a/Engine/Source/Runtime/Experimental/Iris/Core/Private/Iris/Serialization/StringNetSerializerUtils.cpp
+++ b/Engine/Source/Runtime/Experimental/Iris/Core/Private/Iris/Serialization/StringNetSerializerUtils.cpp
@@ -137,11 +137,11 @@ void FStringNetSerializerBase::Dequantize(FNetSerializationContext& Context, con
 			OutLength += 1U;
 		}
 
-		Target = FString(int32(OutLength - 1U), TempString.GetData());
+		Target = FString::ConstructFromPtrSize(TempString.GetData(), int32(OutLength - 1U));
 	}
 	else
 	{
-		Target = FString((Source.ElementCount > 0U ? Source.ElementCount - 1U : 0U), static_cast<ANSICHAR*>(Source.ElementStorage));
+		Target = FString::ConstructFromPtrSize(static_cast<ANSICHAR*>(Source.ElementStorage), (Source.ElementCount > 0U ? Source.ElementCount - 1U : 0U));
 	}
 }
 
diff --git a/Engine/Source/Runtime/GeometryFramework/Private/UDynamicMesh.cpp b/Engine/Source/Runtime/GeometryFramework/Private/UDynamicMesh.cpp
--- a/Engine/Source/Runtime/GeometryFramework/Private/UDynamicMesh.cpp
+++ b/Engine/Source/Runtime/GeometryFramework/Private/UDynamicMesh.cpp
@@ -514,7 +514,7 @@ void UDynamicMesh::ImportCustomProperties(const TCHAR* SourceText, FFeedbackCont
 			if (FoundMeshDataStart)
 			{
 				SourceText = FoundMeshDataStart + FCString::Strlen(MeshDataToken);
-				FString MeshData(MeshDataLen, SourceText);
+				FString MeshData = FString::ConstructFromPtrSize(SourceText, MeshDataLen);
 
 				// fix-up the hack applied to the Base64-encoded string in ExportCustomProperties()
 				for (int32 k = 0; k < MeshData.Len(); ++k)
diff --git a/Engine/Source/Runtime/MovieScene/Private/MovieSceneCommonHelpers.cpp b/Engine/Source/Runtime/MovieScene/Private/MovieSceneCommonHelpers.cpp
--- a/Engine/Source/Runtime/MovieScene/Private/MovieSceneCommonHelpers.cpp
+++ b/Engine/Source/Runtime/MovieScene/Private/MovieSceneCommonHelpers.cpp
@@ -1083,7 +1083,7 @@ FPropertyAndIndex FindPropertyAndArrayIndex(UStruct* InStruct, const FString& Pr
 		int32 OpenIndex = 0;
 		if (PropertyName.FindLastChar('[', OpenIndex))
 		{
-			FString TruncatedPropertyName(OpenIndex, *PropertyName);
+			FString TruncatedPropertyName = FString::ConstructFromPtrSize(*PropertyName, OpenIndex);
 			PropertyAndIndex.Property = FindFProperty<FProperty>(InStruct, *TruncatedPropertyName);
 
 			const int32 NumberLength = PropertyName.Len() - OpenIndex - 2;
diff --git a/Engine/Source/Runtime/Online/BuildPatchServices/Private/Generation/FileAttributesParser.cpp b/Engine/Source/Runtime/Online/BuildPatchServices/Private/Generation/FileAttributesParser.cpp
--- a/Engine/Source/Runtime/Online/BuildPatchServices/Private/Generation/FileAttributesParser.cpp
+++ b/Engine/Source/Runtime/Online/BuildPatchServices/Private/Generation/FileAttributesParser.cpp
@@ -124,10 +124,10 @@ namespace BuildPatchServices
 			// Save position
 			const TCHAR* EndOfLine = CharPtr;
 			// Grab info
-			FString Filename = FString(FilenameEnd - FilenameStart, FilenameStart).Replace(TEXT("\\"), TEXT("/"));
+			FString Filename = FString::ConstructFromPtrSize(FilenameStart, FilenameEnd - FilenameStart).Replace(TEXT("\\"), TEXT("/"));
 			FFileAttributes& FileAttributes = FileAttributesMap.FindOrAdd(Filename);
 			TArray<FString> AttributeParamsArray;
-			FString AttributeParams(EndOfLine - AttributesStart, AttributesStart);
+			FString AttributeParams = FString::ConstructFromPtrSize(AttributesStart, EndOfLine - AttributesStart);
 			AttributeParams.ParseIntoArrayWS(AttributeParamsArray);
 			for (const FString& AttributeParam : AttributeParamsArray)
 			{
diff --git a/Engine/Source/Runtime/Online/HTTPServer/Private/HttpConnectionRequestReadContext.cpp b/Engine/Source/Runtime/Online/HTTPServer/Private/HttpConnectionRequestReadContext.cpp
--- a/Engine/Source/Runtime/Online/HTTPServer/Private/HttpConnectionRequestReadContext.cpp
+++ b/Engine/Source/Runtime/Online/HTTPServer/Private/HttpConnectionRequestReadContext.cpp
@@ -105,7 +105,7 @@ bool FHttpConnectionRequestReadContext::ParseHeader(uint8* ByteBuffer, int32 Buf
 
 			// Build header string
 			FUTF8ToTCHAR WByteBuffer(reinterpret_cast<const ANSICHAR*>(HeaderBytes.GetData()), HeaderBytes.Num());
-			const FString IncomingRequestHeaderStr(WByteBuffer.Length(), WByteBuffer.Get());
+			const FString IncomingRequestHeaderStr = FString::ConstructFromPtrSize(WByteBuffer.Get(), WByteBuffer.Length());
 			Request = BuildRequest(IncomingRequestHeaderStr);
 			if (!Request)
 			{
diff --git a/Engine/Source/Runtime/Online/WebSockets/Private/Lws/LwsWebSocket.cpp b/Engine/Source/Runtime/Online/WebSockets/Private/Lws/LwsWebSocket.cpp
--- a/Engine/Source/Runtime/Online/WebSockets/Private/Lws/LwsWebSocket.cpp
+++ b/Engine/Source/Runtime/Online/WebSockets/Private/Lws/LwsWebSocket.cpp
@@ -476,7 +476,7 @@ int FLwsWebSocket::LwsCallback(lws* Instance, lws_callback_reasons Reason, void*
 		LwsConnection = nullptr;
 
 		FUTF8ToTCHAR Convert((const ANSICHAR*)Data, Length);
-		FString CloseReasonString(Convert.Length(), Convert.Get());
+		FString CloseReasonString = FString::ConstructFromPtrSize(Convert.Get(), Convert.Length());
 		bool bWakeGameThread = false;
 		UE_LOG(LogWebSockets, Verbose, TEXT("FLwsWebSocket[%d]::LwsCallback: Received LWS_CALLBACK_CLIENT_CONNECTION_ERROR, setting State=%s CloseReason=%s PreviousState=%s"),
 			Identifier, ToString(EState::Error), *CloseReasonString, ToString(State));
diff --git a/Engine/Source/Runtime/PropertyPath/Private/PropertyPathHelpers.cpp b/Engine/Source/Runtime/PropertyPath/Private/PropertyPathHelpers.cpp
--- a/Engine/Source/Runtime/PropertyPath/Private/PropertyPathHelpers.cpp
+++ b/Engine/Source/Runtime/PropertyPath/Private/PropertyPathHelpers.cpp
@@ -676,7 +676,7 @@ FPropertyPathSegment::FPropertyPathSegment(int32 InCount, const TCHAR* InString)
 	int32 PropertyNameLength = 0;
 	PropertyPathHelpers::FindFieldNameAndArrayIndex(InCount, InString, PropertyNameLength, &PropertyName, ArrayIndex);
 	ensure(PropertyName != nullptr);
-	FString PropertyNameString(PropertyNameLength, PropertyName);
+	FString PropertyNameString = FString::ConstructFromPtrSize(PropertyName, PropertyNameLength);
 	Name = FName(*PropertyNameString, FNAME_Find);
 }
 
@@ -1047,7 +1047,7 @@ namespace PropertyPathHelpers
 			{
 				OutCount = Offset;
 				// here we need to copy - since we need a section of the string only
-				FString ArrayIndexString(InCount - Offset - 2, &InString[Offset + 1]);
+				FString ArrayIndexString = FString::ConstructFromPtrSize(&InString[Offset + 1], InCount - Offset - 2);
 				OutArrayIndex = FCString::Atoi(*ArrayIndexString);
 				break;
 		}
diff --git a/Engine/Source/Runtime/RenderCore/Private/PipelineCacheUtilities.cpp b/Engine/Source/Runtime/RenderCore/Private/PipelineCacheUtilities.cpp
--- a/Engine/Source/Runtime/RenderCore/Private/PipelineCacheUtilities.cpp
+++ b/Engine/Source/Runtime/RenderCore/Private/PipelineCacheUtilities.cpp
@@ -310,7 +310,7 @@ namespace Private
 
 bool UE::PipelineCacheUtilities::LoadStableKeysFile(const FStringView& Filename, TArray<FStableShaderKeyAndValue>& InOutArray)
 {
-	TUniquePtr<FArchive> FileArchiveInner(IFileManager::Get().CreateFileReader(*FString(Filename.Len(), Filename.GetData())));
+	TUniquePtr<FArchive> FileArchiveInner(IFileManager::Get().CreateFileReader(*FString::ConstructFromPtrSize(Filename.GetData(), Filename.Len())));
 	if (!FileArchiveInner)
 	{
 		return false;
@@ -391,7 +391,7 @@ bool UE::PipelineCacheUtilities::LoadStableKeysFile(const FStringView& Filename,
 
 bool UE::PipelineCacheUtilities::SaveStableKeysFile(const FStringView& Filename, const TSet<FStableShaderKeyAndValue>& Values)
 {
-	TUniquePtr<FArchive> FileArchiveInner(IFileManager::Get().CreateFileWriter(*FString(Filename.Len(), Filename.GetData())));
+	TUniquePtr<FArchive> FileArchiveInner(IFileManager::Get().CreateFileWriter(*FString::ConstructFromPtrSize(Filename.GetData(), Filename.Len())));
 	if (!FileArchiveInner)
 	{
 		return false;
diff --git a/Engine/Source/Runtime/RenderCore/Private/ShaderCore.cpp b/Engine/Source/Runtime/RenderCore/Private/ShaderCore.cpp
--- a/Engine/Source/Runtime/RenderCore/Private/ShaderCore.cpp
+++ b/Engine/Source/Runtime/RenderCore/Private/ShaderCore.cpp
@@ -1509,7 +1509,7 @@ int HandleShaderCompileException(Windows::LPEXCEPTION_POINTERS Info, FString& Ou
 		const TCHAR* CallstackStart = FCString::Strfind(GErrorHist, TEXT("0x"));
 		if (CallstackStart && CallstackStart > GErrorHist)
 		{
-			OutExMsg = FString(CallstackStart - GErrorHist, GErrorHist);
+			OutExMsg = FString::ConstructFromPtrSize(GErrorHist, CallstackStart - GErrorHist);
 			OutCallStack = CallstackStart;
 		}
 		else
@@ -2695,7 +2695,7 @@ static void InternalGetShaderIncludes(const TCHAR* EntryPointVirtualFilePath, co
 				if (IncludeFilenameEnd)
 				{
 					//construct a string between the double quotations
-					FString ExtractedIncludeFilename(FString((int32)(IncludeFilenameEnd - IncludeFilenameBegin - 1), IncludeFilenameBegin + 1));
+					FString ExtractedIncludeFilename = FString::ConstructFromPtrSize(IncludeFilenameBegin + 1, (int32)(IncludeFilenameEnd - IncludeFilenameBegin - 1));
 
 					// If the include is relative, then it must be relative to the current virtual file path.
 					if (!ExtractedIncludeFilename.StartsWith(TEXT("/")))
diff --git a/Engine/Source/Runtime/RenderCore/Private/ShaderPreprocessTypes.cpp b/Engine/Source/Runtime/RenderCore/Private/ShaderPreprocessTypes.cpp
--- a/Engine/Source/Runtime/RenderCore/Private/ShaderPreprocessTypes.cpp
+++ b/Engine/Source/Runtime/RenderCore/Private/ShaderPreprocessTypes.cpp
@@ -323,7 +323,7 @@ void FShaderPreprocessOutput::StripCode(bool bCopyOriginalPreprocessdSource)
 						++Current;
 					}
 
-					FString DirectiveFileName(DirectiveFileNameLen, DirectiveFileNameStart);
+					FString DirectiveFileName = FString::ConstructFromPtrSize(DirectiveFileNameStart, DirectiveFileNameLen);
 
 					// scan to end-of-line and skip past the newline; this would be handled by the newline case above as well,
 					// but we don't want the newline at the end of the line directive to count in our calculated offsets for
diff --git a/Engine/Source/Runtime/Slate/Private/Framework/Text/ShapedTextCache.cpp b/Engine/Source/Runtime/Slate/Private/Framework/Text/ShapedTextCache.cpp
--- a/Engine/Source/Runtime/Slate/Private/Framework/Text/ShapedTextCache.cpp
+++ b/Engine/Source/Runtime/Slate/Private/Framework/Text/ShapedTextCache.cpp
@@ -164,7 +164,7 @@ int32 ShapedTextCacheUtil::FindCharacterIndexAtOffset(const FShapedTextCacheRef&
 	if (GlyphOffsetResultValue.Glyph && GlyphOffsetResultValue.Glyph->NumGraphemeClustersInGlyph > 1)
 	{
 		// Process each grapheme cluster within the ligature
-		const FString LigatureString = FString(GlyphOffsetResultValue.Glyph->NumCharactersInGlyph, InText + GlyphOffsetResultValue.Glyph->SourceIndex);
+		const FString LigatureString = FString::ConstructFromPtrSize(InText + GlyphOffsetResultValue.Glyph->SourceIndex, GlyphOffsetResultValue.Glyph->NumCharactersInGlyph);
 		TSharedRef<IBreakIterator> GraphemeBreakIterator = FBreakIterator::CreateCharacterBoundaryIterator();
 		GraphemeBreakIterator->SetString(LigatureString);
 
diff --git a/Engine/Source/Runtime/Slate/Private/Framework/Text/SlateHyperlinkRun.cpp b/Engine/Source/Runtime/Slate/Private/Framework/Text/SlateHyperlinkRun.cpp
--- a/Engine/Source/Runtime/Slate/Private/Framework/Text/SlateHyperlinkRun.cpp
+++ b/Engine/Source/Runtime/Slate/Private/Framework/Text/SlateHyperlinkRun.cpp
@@ -86,7 +86,7 @@ TSharedRef< ILayoutBlock > FSlateHyperlinkRun::CreateBlock( int32 StartIndex, in
 
 	TSharedRef< SWidget > Widget = SNew( SRichTextHyperlink, ViewModel )
 		.Style( &Style )
-		.Text( FText::FromString( FString( EndIndex - StartIndex, **Text + StartIndex ) ) )
+		.Text( FText::FromString( FString::ConstructFromPtrSize( **Text + StartIndex, EndIndex - StartIndex ) ) )
 		.ToolTip( ToolTip )
 		.ToolTipText( ToolTipText )
 		.OnNavigate( this, &FSlateHyperlinkRun::OnNavigate )
diff --git a/Engine/Source/Runtime/Slate/Private/Framework/Text/TextLayout.cpp b/Engine/Source/Runtime/Slate/Private/Framework/Text/TextLayout.cpp
--- a/Engine/Source/Runtime/Slate/Private/Framework/Text/TextLayout.cpp
+++ b/Engine/Source/Runtime/Slate/Private/Framework/Text/TextLayout.cpp
@@ -610,7 +610,7 @@ void FTextLayout::FlowLineLayout(const int32 LineModelIndex, const float Wrappin
 
 				// This is a single word that's too long to fit onto a line, so we'll fallback to wrapping this word at grapheme cluster boundaries - this may require more than a single split
 				const int32 NonBreakingStringIndexOffset = PreviousBlockEnd;
-				const FString NonBreakingString = FString(Break.TrimmedRange.EndIndex - NonBreakingStringIndexOffset, **LineModel.Text + NonBreakingStringIndexOffset);
+				const FString NonBreakingString = FString::ConstructFromPtrSize(**LineModel.Text + NonBreakingStringIndexOffset, Break.TrimmedRange.EndIndex - NonBreakingStringIndexOffset);
 				GraphemeBreakIterator->SetStringRef(&NonBreakingString);
 
 				CurrentWidth = 0.0f;
@@ -1894,8 +1894,8 @@ bool FTextLayout::SplitLineAt(const FTextLocation& Location)
 
 	FLineModel& LineModel = LineModels[LineIndex];
 
-	FLineModel LeftLineModel(MakeShareable(new FString(BreakLocation, **LineModel.Text)));
-	FLineModel RightLineModel(MakeShareable(new FString(LineModel.Text->Len() - BreakLocation, **LineModel.Text + BreakLocation)));
+	FLineModel LeftLineModel(MakeShareable(new FString(FString::ConstructFromPtrSize(**LineModel.Text, BreakLocation))));
+	FLineModel RightLineModel(MakeShareable(new FString(FString::ConstructFromPtrSize(**LineModel.Text + BreakLocation, LineModel.Text->Len() - BreakLocation))));
 
 	checkf(LeftLineModel.Text->Len() == BreakLocation, TEXT("Debug Source: %s"), *DebugSourceInfo.Get(FString()));
 
diff --git a/Engine/Source/Runtime/SlateCore/Private/Debugging/SlateDebugging.cpp b/Engine/Source/Runtime/SlateCore/Private/Debugging/SlateDebugging.cpp
--- a/Engine/Source/Runtime/SlateCore/Private/Debugging/SlateDebugging.cpp
+++ b/Engine/Source/Runtime/SlateCore/Private/Debugging/SlateDebugging.cpp
@@ -505,7 +505,7 @@ void FSlateDebugging::BroadcastInputEvent(ESlateDebuggingInputEvent InputEventTy
 {
 	if (InReply.IsEventHandled() && InputEvent.IsBound())
 	{
-		InputEvent.Broadcast(FSlateDebuggingInputEventArgs(InputEventType, InInputEvent, InReply, HandlerWidget, FString(1, &AdditionalContent)));
+		InputEvent.Broadcast(FSlateDebuggingInputEventArgs(InputEventType, InInputEvent, InReply, HandlerWidget, FString::ConstructFromPtrSize(&AdditionalContent, 1)));
 	}
 	for (IWidgetInputRoutingEvent* Event : RoutingEvents)
 	{
diff --git a/Engine/Source/Runtime/SlateCore/Private/Input/Events.cpp b/Engine/Source/Runtime/SlateCore/Private/Input/Events.cpp
--- a/Engine/Source/Runtime/SlateCore/Private/Input/Events.cpp
+++ b/Engine/Source/Runtime/SlateCore/Private/Input/Events.cpp
@@ -42,7 +42,7 @@ bool FInputEvent::IsKeyEvent() const
 
 FText FCharacterEvent::ToText() const
 {
-	return FText::Format( NSLOCTEXT("Events", "Char", "Char({0})"), FText::FromString(FString(1, &Character)) );
+	return FText::Format( NSLOCTEXT("Events", "Char", "Char({0})"), FText::FromString(FString::ConstructFromPtrSize(&Character, 1)) );
 }
 
 FText FKeyEvent::ToText() const
diff --git a/Engine/Source/Runtime/StorageServerClient/Private/BuiltInHttpClient/BuiltInHttpClient.cpp b/Engine/Source/Runtime/StorageServerClient/Private/BuiltInHttpClient/BuiltInHttpClient.cpp
--- a/Engine/Source/Runtime/StorageServerClient/Private/BuiltInHttpClient/BuiltInHttpClient.cpp
+++ b/Engine/Source/Runtime/StorageServerClient/Private/BuiltInHttpClient/BuiltInHttpClient.cpp
@@ -193,7 +193,7 @@ IStorageServerHttpClient::FResult FBuiltInHttpClient::RequestSync(
 	}
 	else
 	{
-		FString ErrorMessage = bHasResponsePayload ? FString(ResponsePayload->GetSize(), ANSI_TO_TCHAR(reinterpret_cast<ANSICHAR*>(ResponsePayload->GetData()))) : TEXT("Unknown error");
+		FString ErrorMessage = bHasResponsePayload ? FString::ConstructFromPtrSize(ANSI_TO_TCHAR(reinterpret_cast<ANSICHAR*>(ResponsePayload->GetData())), ResponsePayload->GetSize()) : TEXT("Unknown error");
 
 		return {TIoStatusOr<FIoBuffer>(FIoStatus(ErrorCode, ErrorMessage)), ContentType};
 	}
