Add multi cast delegate for online async op complete event, for monitoring purpose.

#changelist validated

diff --git a/Engine/Plugins/Online/OnlineServices/Source/OnlineServicesCommon/Private/Online/OnlineSyntaxTests.cpp b/Engine/Plugins/Online/OnlineServices/Source/OnlineServicesCommon/Private/Online/OnlineSyntaxTests.cpp
--- a/Engine/Plugins/Online/OnlineServices/Source/OnlineServicesCommon/Private/Online/OnlineSyntaxTests.cpp
+++ b/Engine/Plugins/Online/OnlineServices/Source/OnlineServicesCommon/Private/Online/OnlineSyntaxTests.cpp
@@ -158,6 +158,7 @@ namespace UE::Online {
 
 struct FTestOp
 {
+	static constexpr TCHAR Name[] = TEXT("TestOp");
 	struct Params {};
 	struct Result {};
 };
diff --git a/Engine/Plugins/Online/OnlineServices/Source/OnlineServicesCommon/Public/Online/OnlineAsyncOp.h b/Engine/Plugins/Online/OnlineServices/Source/OnlineServicesCommon/Public/Online/OnlineAsyncOp.h
--- a/Engine/Plugins/Online/OnlineServices/Source/OnlineServicesCommon/Public/Online/OnlineAsyncOp.h
+++ b/Engine/Plugins/Online/OnlineServices/Source/OnlineServicesCommon/Public/Online/OnlineAsyncOp.h
@@ -4,6 +4,7 @@
 
 #include "Online/OnlineAsyncOpHandle.h"
 #include "Online/OnlineResult.h"
+#include "Online/OnlineServicesDelegates.h"
 #include "Online/OnlineTypeInfo.h"
 #include "Async/Async.h"
 #include "Traits/ElementType.h"
@@ -943,6 +944,7 @@ public:
 	TOnlineAsyncOp(FOnlineServicesCommon& InServices, ParamsType&& Params)
 		: Services(InServices)
 		, SharedState(MakeShared<FAsyncOpSharedState>(MoveTemp(Params)))
+		, OpStartTimeAbsoluteSeconds(FPlatformTime::Seconds())
 	{
 	}
 
@@ -1102,6 +1104,8 @@ protected:
 		}
 
 		OnCompleteEvent.Broadcast(*this, Result);
+		double DurationInSeconds = FPlatformTime::Seconds() - OpStartTimeAbsoluteSeconds;
+		OnOnlineAsyncOpCompleted.Broadcast(OpType::Name, Services, Result.IsOk() ? FOnlineError(Errors::ErrorCode::Common::Success) : Result.GetErrorValue(), DurationInSeconds);
 	}
 
 	class FAsyncOpSharedState
@@ -1229,6 +1233,7 @@ protected:
 	TOnlineEventCallable<void(const TOnlineAsyncOp<OpType>&)> OnStartEvent;
 	TOnlineEventCallable<void(const TOnlineAsyncOp<OpType>&, const TOnlineResult<OpType>&)> OnCompleteEvent;
 	int NextStep = 0;
+	double OpStartTimeAbsoluteSeconds;
 
 	friend class FOnlineAsyncOpCache;
 	FOnlineEventDelegateHandle OpCacheHandle; // Delegate handle for FOnlineAsyncOpCache internal usage
diff --git a/Engine/Plugins/Online/OnlineServices/Source/OnlineServicesInterface/Private/Online/OnlineServicesDelegates.cpp b/Engine/Plugins/Online/OnlineServices/Source/OnlineServicesInterface/Private/Online/OnlineServicesDelegates.cpp
--- a/Engine/Plugins/Online/OnlineServices/Source/OnlineServicesInterface/Private/Online/OnlineServicesDelegates.cpp
+++ b/Engine/Plugins/Online/OnlineServices/Source/OnlineServicesInterface/Private/Online/OnlineServicesDelegates.cpp
@@ -5,5 +5,6 @@
 namespace UE::Online {
 
 FOnOnlineServicesCreated OnOnlineServicesCreated;
+FOnOnlineAsyncOpCompleted OnOnlineAsyncOpCompleted;
 
 /* UE::Online */ }
diff --git a/Engine/Plugins/Online/OnlineServices/Source/OnlineServicesInterface/Public/Online/OnlineServicesDelegates.h b/Engine/Plugins/Online/OnlineServices/Source/OnlineServicesInterface/Public/Online/OnlineServicesDelegates.h
--- a/Engine/Plugins/Online/OnlineServices/Source/OnlineServicesInterface/Public/Online/OnlineServicesDelegates.h
+++ b/Engine/Plugins/Online/OnlineServices/Source/OnlineServicesInterface/Public/Online/OnlineServicesDelegates.h
@@ -3,9 +3,13 @@
 #pragma once
 
 #include "Delegates/Delegate.h"
+#include "Online/CoreOnline.h"
 
 namespace UE::Online {
 
+class FOnlineServicesCommon;
+class FOnlineError;
+
 /**
  * Online services delegates that are more external to the online services themselves
  */
@@ -18,6 +22,19 @@ namespace UE::Online {
 DECLARE_MULTICAST_DELEGATE_OneParam(FOnOnlineServicesCreated, TSharedRef<class IOnlineServices> /*NewServices*/);
 extern ONLINESERVICESINTERFACE_API FOnOnlineServicesCreated OnOnlineServicesCreated;
 
+/**
+ * Notification that an online operation has completed
+ *
+ * !!!NOTE!!! The notification can happen on off-game threads, make sure the callbacks are thread-safe
+ *
+ * @param OpName the name of completed operation
+ * @param OnlineServicesCommon the OnlineServices instance
+ * @param OnlineError the result of completed operation
+ * @param DurationInSeconds the duration of the operation from start to complete
+ */
+DECLARE_MULTICAST_DELEGATE_FourParams(FOnOnlineAsyncOpCompleted, const FString& OpName, const FOnlineServicesCommon& OnlineServicesCommon, const UE::Online::FOnlineError& OnlineError, double DurationInSeconds);
+extern ONLINESERVICESINTERFACE_API FOnOnlineAsyncOpCompleted OnOnlineAsyncOpCompleted;
+
 /* UE::Online */ }
 
 #if UE_ENABLE_INCLUDE_ORDER_DEPRECATED_IN_5_2
