Regster OnlineService presence exec commands.
+ Fix missing param from GetCachedPrersence meta.

diff --git a/Engine/Plugins/Online/OnlineServices/Source/OnlineServicesCommon/Private/Online/PresenceCommon.cpp b/Engine/Plugins/Online/OnlineServices/Source/OnlineServicesCommon/Private/Online/PresenceCommon.cpp
--- a/Engine/Plugins/Online/OnlineServices/Source/OnlineServicesCommon/Private/Online/PresenceCommon.cpp
+++ b/Engine/Plugins/Online/OnlineServices/Source/OnlineServicesCommon/Private/Online/PresenceCommon.cpp
@@ -1,7 +1,7 @@
 // Copyright Epic Games, Inc. All Rights Reserved.
 
 #include "Online/PresenceCommon.h"
-
+#include "Online/OnlineUtils.h"
 
 namespace UE::Online {
 
@@ -11,6 +11,15 @@ FPresenceCommon::FPresenceCommon(FOnlineServicesCommon& InServices)
 {
 }
 
+void FPresenceCommon::RegisterCommands()
+{
+	RegisterCommand(&FPresenceCommon::QueryPresence);
+	RegisterCommand(&FPresenceCommon::BatchQueryPresence);
+	RegisterCommand(&FPresenceCommon::GetCachedPresence);
+	RegisterCommand(&FPresenceCommon::UpdatePresence);
+	RegisterCommand(&FPresenceCommon::PartialUpdatePresence);
+}
+
 TOnlineAsyncOpHandle<FQueryPresence> FPresenceCommon::QueryPresence(FQueryPresence::Params&& Params)
 {
 	TOnlineAsyncOpRef<FQueryPresence> Operation = GetOp<FQueryPresence>(MoveTemp(Params));
diff --git a/Engine/Plugins/Online/OnlineServices/Source/OnlineServicesCommon/Public/Online/PresenceCommon.h b/Engine/Plugins/Online/OnlineServices/Source/OnlineServicesCommon/Public/Online/PresenceCommon.h
--- a/Engine/Plugins/Online/OnlineServices/Source/OnlineServicesCommon/Public/Online/PresenceCommon.h
+++ b/Engine/Plugins/Online/OnlineServices/Source/OnlineServicesCommon/Public/Online/PresenceCommon.h
@@ -16,6 +16,9 @@ public:
 
 	FPresenceCommon(FOnlineServicesCommon& InServices);
 
+	// IOnlineComponent
+	virtual void RegisterCommands() override;
+
 	// IPresence
 	virtual TOnlineAsyncOpHandle<FQueryPresence> QueryPresence(FQueryPresence::Params&& Params) override;
 	virtual TOnlineAsyncOpHandle<FBatchQueryPresence> BatchQueryPresence(FBatchQueryPresence::Params&& Params) override;
diff --git a/Engine/Plugins/Online/OnlineServices/Source/OnlineServicesInterface/Public/Online/Presence.h b/Engine/Plugins/Online/OnlineServices/Source/OnlineServicesInterface/Public/Online/Presence.h
--- a/Engine/Plugins/Online/OnlineServices/Source/OnlineServicesInterface/Public/Online/Presence.h
+++ b/Engine/Plugins/Online/OnlineServices/Source/OnlineServicesInterface/Public/Online/Presence.h
@@ -389,7 +389,8 @@ BEGIN_ONLINE_STRUCT_META(FBatchQueryPresence::Result)
 END_ONLINE_STRUCT_META()
 
 BEGIN_ONLINE_STRUCT_META(FGetCachedPresence::Params)
-	ONLINE_STRUCT_FIELD(FGetCachedPresence::Params, LocalAccountId)
+	ONLINE_STRUCT_FIELD(FGetCachedPresence::Params, LocalAccountId),
+	ONLINE_STRUCT_FIELD(FGetCachedPresence::Params, TargetAccountId)
 END_ONLINE_STRUCT_META()
 
 BEGIN_ONLINE_STRUCT_META(FGetCachedPresence::Result)
