- Add "ConfigName" param to UE::Online::GetServices
- Add "Type" to CommonConfig, so we can configure Server and Client contexts with separate ConfigName.

#rnx
#tests local PIE testing and 2 other devs unshelved and used to unblock work

diff --git a/Engine/Plugins/Online/OnlineServices/Source/OnlineServicesCommon/Private/Online/OnlineServicesCommon.cpp b/Engine/Plugins/Online/OnlineServices/Source/OnlineServicesCommon/Private/Online/OnlineServicesCommon.cpp
--- a/Engine/Plugins/Online/OnlineServices/Source/OnlineServicesCommon/Private/Online/OnlineServicesCommon.cpp
+++ b/Engine/Plugins/Online/OnlineServices/Source/OnlineServicesCommon/Private/Online/OnlineServicesCommon.cpp
@@ -23,12 +23,13 @@ END_ONLINE_STRUCT_META()
 
 uint32 FOnlineServicesCommon::NextInstanceIndex = 0;
 
-FOnlineServicesCommon::FOnlineServicesCommon(const FString& InConfigName, FName InInstanceName)
-	: OpCache(InConfigName, *this)
+FOnlineServicesCommon::FOnlineServicesCommon(const FString& InServiceConfigName, FName InInstanceName, FName InInstanceConfigName)
+	: OpCache(InServiceConfigName, *this)
 	, InstanceIndex(NextInstanceIndex++)
 	, InstanceName(InInstanceName)
+	, InstanceConfigName(InInstanceConfigName)
 	, ConfigProvider(MakeUnique<FOnlineConfigProviderGConfig>(GEngineIni))
-	, ConfigName(InConfigName)
+	, ServiceConfigName(InServiceConfigName)
 	, SerialQueue(ParallelQueue)
 {
 }
@@ -140,6 +141,11 @@ FName FOnlineServicesCommon::GetInstanceName() const
 	return InstanceName;
 }
 
+FName FOnlineServicesCommon::GetInstanceConfigName() const
+{
+	return InstanceConfigName;
+}
+
 void FOnlineServicesCommon::AssignBaseInterfaceSharedPtr(const FOnlineTypeName& TypeName, void* OutBaseInterfaceSP)
 {
 	Components.AssignBaseSharedPtr(TypeName, OutBaseInterfaceSP);
@@ -234,7 +240,8 @@ bool FOnlineServicesCommon::Exec(UWorld* World, const TCHAR* Cmd, FOutputDevice&
 		}
 		else if (FParse::Command(&Cmd, TEXT("List")))
 		{
-			Ar.Logf(TEXT("%u: %s"), InstanceIndex, *GetConfigName());
+			Ar.Logf(TEXT("%u: ServiceConfigName=[%s] InstanceName=[%s] InstanceConfigName=[%s]"),
+				InstanceIndex, *GetServiceConfigName(), *GetInstanceName().ToString(), *GetInstanceConfigName().ToString());
 		}
 	}
 	return false;
diff --git a/Engine/Plugins/Online/OnlineServices/Source/OnlineServicesCommon/Private/Online/OnlineSyntaxTests.cpp b/Engine/Plugins/Online/OnlineServices/Source/OnlineServicesCommon/Private/Online/OnlineSyntaxTests.cpp
--- a/Engine/Plugins/Online/OnlineServices/Source/OnlineServicesCommon/Private/Online/OnlineSyntaxTests.cpp
+++ b/Engine/Plugins/Online/OnlineServices/Source/OnlineServicesCommon/Private/Online/OnlineSyntaxTests.cpp
@@ -256,7 +256,7 @@ class FOnlineServicesTest : public UE::Online::FOnlineServicesCommon
 public:
 	using Super = UE::Online::FOnlineServicesCommon;
 
-	FOnlineServicesTest(FName InInstanceName) : Super(TEXT("Test"), InInstanceName) {}
+	FOnlineServicesTest(FName InInstanceName, FName InInstanceConfigName) : Super(TEXT("Test"), InInstanceName, InInstanceConfigName) {}
 
 	virtual UE::Online::EOnlineServices GetServicesProvider() const override { return UE::Online::EOnlineServices::Null; }
 
@@ -300,7 +300,7 @@ inline void DelegateTest()
 	const UMyObject* ConstUObject = UObject;
 #endif // TEST_CONSTRUCT_DELEGATE_SYNTAX_UOBJECT
 
-	TSharedRef<FOnlineServicesTest> ServicesTest = MakeShared<FOnlineServicesTest>(NAME_None);
+	TSharedRef<FOnlineServicesTest> ServicesTest = MakeShared<FOnlineServicesTest>(NAME_None, NAME_None);
 	FOnlineComponentTest* TestComponent = ServicesTest->Get<FOnlineComponentTest>();
 	const FOnlineComponentTest* ConstTestComponent = ServicesTest->Get<FOnlineComponentTest>();
 
diff --git a/Engine/Plugins/Online/OnlineServices/Source/OnlineServicesCommon/Public/Online/OnlineComponent.h b/Engine/Plugins/Online/OnlineServices/Source/OnlineServicesCommon/Public/Online/OnlineComponent.h
--- a/Engine/Plugins/Online/OnlineServices/Source/OnlineServicesCommon/Public/Online/OnlineComponent.h
+++ b/Engine/Plugins/Online/OnlineServices/Source/OnlineServicesCommon/Public/Online/OnlineComponent.h
@@ -102,7 +102,7 @@ public:
 		TArray<FString> SectionHeiarchy;
 		FString SectionName = TEXT("OnlineServices");
 		SectionHeiarchy.Add(SectionName);
-		SectionName += TEXT(".") + GetServices().GetConfigName();
+		SectionName += TEXT(".") + GetServices().GetServiceConfigName();
 		SectionHeiarchy.Add(SectionName);
 		SectionName += TEXT(".") + GetConfigName();
 		SectionHeiarchy.Add(SectionName);
diff --git a/Engine/Plugins/Online/OnlineServices/Source/OnlineServicesCommon/Public/Online/OnlineServicesCommon.h b/Engine/Plugins/Online/OnlineServices/Source/OnlineServicesCommon/Public/Online/OnlineServicesCommon.h
--- a/Engine/Plugins/Online/OnlineServices/Source/OnlineServicesCommon/Public/Online/OnlineServicesCommon.h
+++ b/Engine/Plugins/Online/OnlineServices/Source/OnlineServicesCommon/Public/Online/OnlineServicesCommon.h
@@ -24,7 +24,7 @@ class ONLINESERVICESCOMMON_API FOnlineServicesCommon
 public:
 	using Super = IOnlineServices;
 
-	FOnlineServicesCommon(const FString& InConfigName, FName InInstanceName);
+	FOnlineServicesCommon(const FString& InServiceConfigName, FName InInstanceName, FName InInstanceConfigName = NAME_None);
 	FOnlineServicesCommon(const FOnlineServicesCommon&) = delete;
 	FOnlineServicesCommon(FOnlineServicesCommon&&) = delete;
 	virtual ~FOnlineServicesCommon();
@@ -49,6 +49,7 @@ public:
 	virtual IUserFilePtr GetUserFileInterface() override;
 	virtual TOnlineResult<FGetResolvedConnectString> GetResolvedConnectString(FGetResolvedConnectString::Params&& Params) override;
 	virtual FName GetInstanceName() const override;
+	virtual FName GetInstanceConfigName() const override;
 	virtual void AssignBaseInterfaceSharedPtr(const FOnlineTypeName& TypeName, void* OutBaseInterfaceSP) override final;
 
 	// FOnlineServicesCommon
@@ -168,16 +169,18 @@ public:
 	}
 
 	/**
-	 * Get the config name for the Subsystem
+	 * Get the ini config name for the Subsystem
 	 */
-	const FString& GetConfigName() const { return ConfigName; }
+	UE_DEPRECATED(5.5, "GetConfigName has been renamed GetServiceConfigName")
+	const FString& GetConfigName() const { return ServiceConfigName; }
+	const FString& GetServiceConfigName() const { return ServiceConfigName; }
 
 	TArray<FString> GetConfigSectionHeiarchy(const FString& OperationName = FString()) const
 	{
 		TArray<FString> SectionHeiarchy;
 		FString SectionName = TEXT("OnlineServices");
 		SectionHeiarchy.Add(SectionName);
-		SectionName += TEXT(".") + GetConfigName();
+		SectionName += TEXT(".") + GetServiceConfigName();
 		SectionHeiarchy.Add(SectionName);
 		if (!OperationName.IsEmpty())
 		{
@@ -212,7 +215,7 @@ public:
 		{
 			SectionHeiarchy.Add(SectionName + TEXT(".") + InterfaceName);
 		}
-		SectionName += TEXT(".") + GetConfigName();
+		SectionName += TEXT(".") + GetServiceConfigName();
 		SectionHeiarchy.Add(SectionName);
 		if (!InterfaceName.IsEmpty())
 		{
@@ -327,13 +330,14 @@ protected:
 	static uint32 NextInstanceIndex;
 	uint32 InstanceIndex;
 	FName InstanceName;
+	FName InstanceConfigName;
 
 	FOnlineComponentRegistry Components;
 	TUniquePtr<IOnlineConfigProvider> ConfigProvider;
 
 	/* Config section overrides */
 	TArray<FString> ConfigSectionOverrides;
-	FString ConfigName;
+	FString ServiceConfigName;
 
 	FOnlineAsyncOpQueueParallel ParallelQueue;
 	FOnlineAsyncOpQueueSerial SerialQueue;
diff --git a/Engine/Plugins/Online/OnlineServices/Source/OnlineServicesInterface/Private/Online/OnlineServices.cpp b/Engine/Plugins/Online/OnlineServices/Source/OnlineServicesInterface/Private/Online/OnlineServices.cpp
--- a/Engine/Plugins/Online/OnlineServices/Source/OnlineServicesInterface/Private/Online/OnlineServices.cpp
+++ b/Engine/Plugins/Online/OnlineServices/Source/OnlineServicesInterface/Private/Online/OnlineServices.cpp
@@ -70,14 +70,14 @@ bool IsLoaded(EOnlineServices OnlineServices, FName InstanceName)
 	return FOnlineServicesRegistry::Get().IsLoaded(OnlineServices, InstanceName);
 }
 
-TSharedPtr<IOnlineServices> GetServices(EOnlineServices OnlineServices, FName InstanceName)
+TSharedPtr<IOnlineServices> GetServices(EOnlineServices OnlineServices, FName InstanceName, FName InstanceConfigName)
 {
-	return FOnlineServicesRegistry::Get().GetNamedServicesInstance(OnlineServices, InstanceName);
+	return FOnlineServicesRegistry::Get().GetNamedServicesInstance(OnlineServices, InstanceName, InstanceConfigName);
 }
 
-void DestroyService(EOnlineServices OnlineServices, FName InstanceName)
+void DestroyService(EOnlineServices OnlineServices, FName InstanceName, FName InstanceConfigName)
 {
-	FOnlineServicesRegistry::Get().DestroyNamedServicesInstance(OnlineServices, InstanceName);
+	FOnlineServicesRegistry::Get().DestroyNamedServicesInstance(OnlineServices, InstanceName, InstanceConfigName);
 }
 
 void DestroyAllNamedServices(EOnlineServices OnlineServices)
diff --git a/Engine/Plugins/Online/OnlineServices/Source/OnlineServicesInterface/Private/Online/OnlineServicesRegistry.cpp b/Engine/Plugins/Online/OnlineServices/Source/OnlineServicesInterface/Private/Online/OnlineServicesRegistry.cpp
--- a/Engine/Plugins/Online/OnlineServices/Source/OnlineServicesInterface/Private/Online/OnlineServicesRegistry.cpp
+++ b/Engine/Plugins/Online/OnlineServices/Source/OnlineServicesInterface/Private/Online/OnlineServicesRegistry.cpp
@@ -74,19 +74,19 @@ void FOnlineServicesRegistry::UnregisterServicesFactory(EOnlineServices OnlineSe
 	DestroyAllNamedServicesInstances(OnlineServices);
 }
 
-bool FOnlineServicesRegistry::IsLoaded(EOnlineServices OnlineServices, FName InstanceName) const
+bool FOnlineServicesRegistry::IsLoaded(EOnlineServices OnlineServices, FName InstanceName, FName InstanceConfigName) const
 {
 	OnlineServices = ResolveServiceName(OnlineServices);
 
 	bool bExists = false;
-	if (const TMap<FName, TSharedRef<IOnlineServices>>* OnlineServicesInstances = NamedServiceInstances.Find(OnlineServices))
+	if (const TMap<FInstanceNameInstanceConfigNamePair, TSharedRef<IOnlineServices>>* OnlineServicesInstances = NamedServiceInstances.Find(OnlineServices))
 	{
-		bExists = OnlineServicesInstances->Find(InstanceName) != nullptr;
+		bExists = OnlineServicesInstances->Find({ InstanceName, InstanceConfigName }) != nullptr;
 	}
 	return bExists;
 }
 
-TSharedPtr<IOnlineServices> FOnlineServicesRegistry::GetNamedServicesInstance(EOnlineServices OnlineServices, FName InstanceName)
+TSharedPtr<IOnlineServices> FOnlineServicesRegistry::GetNamedServicesInstance(EOnlineServices OnlineServices, FName InstanceName, FName InstanceConfigName)
 {
 	OnlineServices = ResolveServiceName(OnlineServices);
 
@@ -94,16 +94,16 @@ TSharedPtr<IOnlineServices> FOnlineServicesRegistry::GetNamedServicesInstance(EO
 
 	if (OnlineServices < EOnlineServices::None)
 	{
-		if (TSharedRef<IOnlineServices>* ServicesPtr = NamedServiceInstances.FindOrAdd(OnlineServices).Find(InstanceName))
+		if (TSharedRef<IOnlineServices>* ServicesPtr = NamedServiceInstances.FindOrAdd(OnlineServices).Find({ InstanceName, InstanceConfigName }))
 		{
 			Services = *ServicesPtr;
 		}
 		else
 		{
-			Services = CreateServices(OnlineServices, InstanceName);
+			Services = CreateServices(OnlineServices, InstanceName, InstanceConfigName);
 			if (Services.IsValid())
 			{
-				NamedServiceInstances.FindOrAdd(OnlineServices).Add(InstanceName, Services.ToSharedRef());
+				NamedServiceInstances.FindOrAdd(OnlineServices).Add({ InstanceName, InstanceConfigName }, Services.ToSharedRef());
 				OnOnlineServicesCreated.Broadcast(Services.ToSharedRef());
 			}
 		}
@@ -125,15 +125,15 @@ void FOnlineServicesRegistry::ClearDefaultServiceOverride()
 }
 #endif //WITH_DEV_AUTOMATION_TESTS
 
-void FOnlineServicesRegistry::DestroyNamedServicesInstance(EOnlineServices OnlineServices, FName InstanceName)
+void FOnlineServicesRegistry::DestroyNamedServicesInstance(EOnlineServices OnlineServices, FName InstanceName, FName InstanceConfigName)
 {
 	OnlineServices = ResolveServiceName(OnlineServices);
 
-	if (TSharedRef<IOnlineServices>* ServicesPtr = NamedServiceInstances.FindOrAdd(OnlineServices).Find(InstanceName))
+	if (TSharedRef<IOnlineServices>* ServicesPtr = NamedServiceInstances.FindOrAdd(OnlineServices).Find({ InstanceName, InstanceConfigName }))
 	{
 		(*ServicesPtr)->Destroy();
 
-		NamedServiceInstances.FindOrAdd(OnlineServices).Remove(InstanceName);
+		NamedServiceInstances.FindOrAdd(OnlineServices).Remove({ InstanceName, InstanceConfigName });
 	}
 }
 
@@ -141,9 +141,9 @@ void FOnlineServicesRegistry::DestroyAllNamedServicesInstances(EOnlineServices O
 {
 	OnlineServices = ResolveServiceName(OnlineServices);
 
-	if (TMap<FName, TSharedRef<IOnlineServices>>* ServicesMapPtr = NamedServiceInstances.Find(OnlineServices))
+	if (TMap<FInstanceNameInstanceConfigNamePair, TSharedRef<IOnlineServices>>* ServicesMapPtr = NamedServiceInstances.Find(OnlineServices))
 	{
-		for (const TPair<FName, TSharedRef<IOnlineServices>>& ServicesEntryRef : *ServicesMapPtr)
+		for (const TPair<FInstanceNameInstanceConfigNamePair, TSharedRef<IOnlineServices>>& ServicesEntryRef : *ServicesMapPtr)
 		{
 			ServicesEntryRef.Value->Destroy();
 		}
@@ -154,22 +154,26 @@ void FOnlineServicesRegistry::DestroyAllNamedServicesInstances(EOnlineServices O
 
 void FOnlineServicesRegistry::DestroyAllServicesInstancesWithName(FName InstanceName)
 {
-	for (auto It = NamedServiceInstances.CreateIterator(); It; ++It)
+	for (auto NamedServiceIterator = NamedServiceInstances.CreateIterator(); NamedServiceIterator; ++NamedServiceIterator)
 	{
-		if (TSharedRef<IOnlineServices>* Service = It->Value.Find(InstanceName))
+		for (auto InstanceIterator = NamedServiceIterator->Value.CreateIterator(); InstanceIterator; ++InstanceIterator)
 		{
-			(*Service)->Destroy();
-			It->Value.Remove(InstanceName);
+			if (InstanceIterator->Key.Key == InstanceName)
+			{
+				TSharedRef<IOnlineServices> Service = InstanceIterator->Value;
+				Service->Destroy();
+				InstanceIterator.RemoveCurrent();
+			}
 		}
 
-		if (It->Value.IsEmpty())
+		if (NamedServiceIterator->Value.IsEmpty())
 		{
-			It.RemoveCurrent();
+			NamedServiceIterator.RemoveCurrent();
 		}
 	}
 }
 
-TSharedPtr<IOnlineServices> FOnlineServicesRegistry::CreateServices(EOnlineServices OnlineServices, FName InstanceName)
+TSharedPtr<IOnlineServices> FOnlineServicesRegistry::CreateServices(EOnlineServices OnlineServices, FName InstanceName, FName InstanceConfigName)
 {
 	OnlineServices = ResolveServiceName(OnlineServices);
 
@@ -178,7 +182,7 @@ TSharedPtr<IOnlineServices> FOnlineServicesRegistry::CreateServices(EOnlineServi
 	FFactoryAndPriority* FactoryAndPriority = ServicesFactories.Find(OnlineServices);
 	if (FactoryAndPriority != nullptr)
 	{
-		Services = FactoryAndPriority->Factory->Create(InstanceName);
+		Services = FactoryAndPriority->Factory->Create(InstanceName, InstanceConfigName);
 		Services->Init();
 	}
 
@@ -187,9 +191,9 @@ TSharedPtr<IOnlineServices> FOnlineServicesRegistry::CreateServices(EOnlineServi
 
 void FOnlineServicesRegistry::GetAllServicesInstances(TArray<TSharedRef<IOnlineServices>>& OutOnlineServices) const
 {
-	for (const TPair<EOnlineServices, TMap<FName, TSharedRef<IOnlineServices>>>& OnlineServiceTypesMaps : NamedServiceInstances)
+	for (const TPair<EOnlineServices, TMap<FInstanceNameInstanceConfigNamePair, TSharedRef<IOnlineServices>>>& OnlineServiceTypesMaps : NamedServiceInstances)
 	{
-		for (const TPair<FName, TSharedRef<IOnlineServices>>& NamedInstance : OnlineServiceTypesMaps.Value)
+		for (const TPair<FInstanceNameInstanceConfigNamePair, TSharedRef<IOnlineServices>>& NamedInstance : OnlineServiceTypesMaps.Value)
 		{
 			OutOnlineServices.Emplace(NamedInstance.Value);
 		}
@@ -198,9 +202,9 @@ void FOnlineServicesRegistry::GetAllServicesInstances(TArray<TSharedRef<IOnlineS
 
 FOnlineServicesRegistry::~FOnlineServicesRegistry()
 {
-	for (TPair<EOnlineServices, TMap<FName, TSharedRef<IOnlineServices>>>& ServiceInstances : NamedServiceInstances)
+	for (TPair<EOnlineServices, TMap<FInstanceNameInstanceConfigNamePair, TSharedRef<IOnlineServices>>>& ServiceInstances : NamedServiceInstances)
 	{
-		for (TPair<FName, TSharedRef<IOnlineServices>>& ServiceInstance : ServiceInstances.Value)
+		for (TPair<FInstanceNameInstanceConfigNamePair, TSharedRef<IOnlineServices>>& ServiceInstance : ServiceInstances.Value)
 		{
 			ServiceInstance.Value->Destroy();
 		}
diff --git a/Engine/Plugins/Online/OnlineServices/Source/OnlineServicesInterface/Public/Online/OnlineServices.h b/Engine/Plugins/Online/OnlineServices/Source/OnlineServicesInterface/Public/Online/OnlineServices.h
--- a/Engine/Plugins/Online/OnlineServices/Source/OnlineServicesInterface/Public/Online/OnlineServices.h
+++ b/Engine/Plugins/Online/OnlineServices/Source/OnlineServicesInterface/Public/Online/OnlineServices.h
@@ -160,6 +160,12 @@ public:
 	 */
 	virtual FName GetInstanceName() const = 0;
 
+	/**
+	* Get the instance config name for this instance. 
+	* @return instance config name
+	*/
+	virtual FName GetInstanceConfigName() const = 0;
+
 	/**
 	 * Retrieve any of the standard or implementation specific interfaces
 	 */
@@ -206,9 +212,10 @@ ONLINESERVICESINTERFACE_API bool IsLoaded(EOnlineServices OnlineServices = EOnli
  *
  * @param OnlineServices Type of online services to retrieve
  * @param InstanceName Name of the services instance to retrieve
+ * @param InstanceConfigName Name of the implementation specific config to use for this services instance
  * @return The services instance or an invalid pointer if the services is unavailable
  */
-ONLINESERVICESINTERFACE_API TSharedPtr<IOnlineServices> GetServices(EOnlineServices OnlineServices = EOnlineServices::Default, FName InstanceName = NAME_None);
+ONLINESERVICESINTERFACE_API TSharedPtr<IOnlineServices> GetServices(EOnlineServices OnlineServices = EOnlineServices::Default, FName InstanceName = NAME_None, FName InstanceConfigName = NAME_None);
 
 /**
  * Get a specific services type and cast to the specific services type
@@ -217,9 +224,9 @@ ONLINESERVICESINTERFACE_API TSharedPtr<IOnlineServices> GetServices(EOnlineServi
  * @return The services instance or an invalid pointer if the services is unavailable
  */
 template <typename ServicesClass>
-TSharedPtr<ServicesClass> GetServices(FName InstanceName = NAME_None)
+TSharedPtr<ServicesClass> GetServices(FName InstanceName = NAME_None, FName InstanceConfigName = NAME_None)
 {
-	return StaticCastSharedPtr<ServicesClass>(GetServices(ServicesClass::GetServicesProvider(), InstanceName));
+	return StaticCastSharedPtr<ServicesClass>(GetServices(ServicesClass::GetServicesProvider(), InstanceName, InstanceConfigName));
 }
 
 /**
@@ -228,7 +235,7 @@ TSharedPtr<ServicesClass> GetServices(FName InstanceName = NAME_None)
  * @param OnlineServices Type of online services to destroy
  * @param InstanceName Name of the services instance to destroy
  */
-ONLINESERVICESINTERFACE_API void DestroyService(EOnlineServices OnlineServices = EOnlineServices::Default, FName InstanceName = NAME_None);
+ONLINESERVICESINTERFACE_API void DestroyService(EOnlineServices OnlineServices = EOnlineServices::Default, FName InstanceName = NAME_None, FName InstanceConfigName = NAME_None);
 
 /**
  * Destroy all instances of the online services specified by EOnlineServices type
diff --git a/Engine/Plugins/Online/OnlineServices/Source/OnlineServicesInterface/Public/Online/OnlineServicesRegistry.h b/Engine/Plugins/Online/OnlineServices/Source/OnlineServicesInterface/Public/Online/OnlineServicesRegistry.h
--- a/Engine/Plugins/Online/OnlineServices/Source/OnlineServicesInterface/Public/Online/OnlineServicesRegistry.h
+++ b/Engine/Plugins/Online/OnlineServices/Source/OnlineServicesInterface/Public/Online/OnlineServicesRegistry.h
@@ -26,7 +26,9 @@ public:
 	 *
 	 * @return Initialized IOnlineServices instance
 	 */
-	virtual TSharedPtr<IOnlineServices> Create(FName InstanceName) = 0;
+	UE_DEPRECATED(5.5, "Please use the new Create taking an additional InstanceConfigName param")
+	virtual TSharedPtr<IOnlineServices> Create(FName InstanceName) final { return Create(InstanceName, NAME_None); }
+	virtual TSharedPtr<IOnlineServices> Create(FName InstanceName, FName InstanceConfigName) = 0;
 };
 
 class FOnlineServicesRegistry
@@ -66,28 +68,35 @@ public:
 	 *
 	 * @param OnlineServices Type of online services for the IOnlineServices instance
 	 * @param InstanceName Name of the instance
+	 * @param InstanceConfigName Name of the config
 	 *
 	 * @return true if the instance is loaded
 	 */
-	ONLINESERVICESINTERFACE_API bool IsLoaded(EOnlineServices OnlineServices, FName InstanceName) const;
+	ONLINESERVICESINTERFACE_API bool IsLoaded(EOnlineServices OnlineServices, FName InstanceName, FName InstanceConfigName = NAME_None) const;
 
 	/**
 	 * Get a named instance of a specific IOnlineServices
 	 * 
 	 * @param OnlineServices Type of online services for the IOnlineServices instance
 	 * @param InstanceName Name of the instance
+	 * @param InstanceConfigName Name of the config to use
 	 * 
 	 * @return The services instance, or an invalid pointer if the OnlineServices is unavailable
 	 */
-	ONLINESERVICESINTERFACE_API TSharedPtr<IOnlineServices> GetNamedServicesInstance(EOnlineServices OnlineServices, FName InstanceName);
+	UE_DEPRECATED(5.5, "Please call the new GetNamedServicesInstance which takes an additional InstanceConfigName param")
+	TSharedPtr<IOnlineServices> GetNamedServicesInstance(EOnlineServices OnlineServices, FName InstanceName) { return GetNamedServicesInstance(OnlineServices, InstanceName, NAME_None); }
+	ONLINESERVICESINTERFACE_API TSharedPtr<IOnlineServices> GetNamedServicesInstance(EOnlineServices OnlineServices, FName InstanceName, FName InstanceConfigName);
 
 	/**
 	 * Destroy a named instance of a specific OnlineServices
 	 *
-	 * @param OnlineServices  Type of online services for the IOnlineServices instance
+	 * @param OnlineServices Type of online services for the IOnlineServices instance
 	 * @param InstanceName Name of the instance
+	 * @param InstanceConfigName Name of the config
 	 */
-	ONLINESERVICESINTERFACE_API void DestroyNamedServicesInstance(EOnlineServices OnlineServices, FName InstanceName);
+	UE_DEPRECATED(5.5, "Please call the new DestroyNamedServicesInstance which takes an additional InstanceConfigName param")
+	void DestroyNamedServicesInstance(EOnlineServices OnlineServices, FName InstanceName) { return DestroyNamedServicesInstance(OnlineServices, InstanceName, NAME_None); }
+	ONLINESERVICESINTERFACE_API void DestroyNamedServicesInstance(EOnlineServices OnlineServices, FName InstanceName, FName InstanceConfigName);
 
 	/**
 	 * Destroy all instances of a specific OnlineServices
@@ -107,10 +116,14 @@ public:
 	 * Create and initialize a new IOnlineServices instance
 	 *
 	 * @param OnlineServices Type of online services for the IOnlineServices instance
+	 * @param InstanceName Name of the instance
+	 * @param InstanceConfigName Name of the config
 	 * 
 	 * @return The initialized IOnlineServices instance, or an invalid pointer if the OnlineServices is unavailable
 	 */
-	ONLINESERVICESINTERFACE_API TSharedPtr<IOnlineServices> CreateServices(EOnlineServices OnlineServices, FName InstanceName);
+	UE_DEPRECATED(5.5, "Please call the new CreateServices which takes an additional InstanceConfigName param")
+	TSharedPtr<IOnlineServices> CreateServices(EOnlineServices OnlineServices, FName InstanceName) { return CreateServices(OnlineServices, InstanceName, NAME_None); }
+	ONLINESERVICESINTERFACE_API TSharedPtr<IOnlineServices> CreateServices(EOnlineServices OnlineServices, FName InstanceName, FName InstanceConfigName);
 
 	/**
 	 * Get list of all instantiated OnlineServices
@@ -136,8 +149,8 @@ private:
 	struct FFactoryAndPriority
 	{
 		FFactoryAndPriority(TUniquePtr<IOnlineServicesFactory>&& InFactory, int32 InPriority)
-			: Factory(MoveTemp(InFactory))
-			, Priority(InPriority)
+		: Factory(MoveTemp(InFactory))
+		, Priority(InPriority)
 		{
 		}
 
@@ -146,7 +159,9 @@ private:
 	};
 
 	TMap<EOnlineServices, FFactoryAndPriority> ServicesFactories;
-	TMap<EOnlineServices, TMap<FName, TSharedRef<IOnlineServices>>> NamedServiceInstances;
+
+	using FInstanceNameInstanceConfigNamePair = TPair<FName, FName>;
+	TMap<EOnlineServices, TMap<FInstanceNameInstanceConfigNamePair, TSharedRef<IOnlineServices>>> NamedServiceInstances;
 	EOnlineServices DefaultServiceOverride = EOnlineServices::Default;
 
 	FOnlineServicesRegistry() {}
diff --git a/Engine/Plugins/Online/OnlineServicesEOS/Source/Private/Online/OnlineServicesEOS.cpp b/Engine/Plugins/Online/OnlineServicesEOS/Source/Private/Online/OnlineServicesEOS.cpp
--- a/Engine/Plugins/Online/OnlineServicesEOS/Source/Private/Online/OnlineServicesEOS.cpp
+++ b/Engine/Plugins/Online/OnlineServicesEOS/Source/Private/Online/OnlineServicesEOS.cpp
@@ -11,8 +11,8 @@
 
 namespace UE::Online {
 
-FOnlineServicesEOS::FOnlineServicesEOS(FName InInstanceName)
-	: Super(InInstanceName)
+FOnlineServicesEOS::FOnlineServicesEOS(FName InInstanceName, FName InInstanceConfigName)
+	: Super(InInstanceName, InInstanceConfigName)
 {
 }
 
diff --git a/Engine/Plugins/Online/OnlineServicesEOS/Source/Private/Online/OnlineServicesEOSModule.cpp b/Engine/Plugins/Online/OnlineServicesEOS/Source/Private/Online/OnlineServicesEOSModule.cpp
--- a/Engine/Plugins/Online/OnlineServicesEOS/Source/Private/Online/OnlineServicesEOSModule.cpp
+++ b/Engine/Plugins/Online/OnlineServicesEOS/Source/Private/Online/OnlineServicesEOSModule.cpp
@@ -16,9 +16,9 @@ class FOnlineServicesFactoryEOS : public IOnlineServicesFactory
 {
 public:
 	virtual ~FOnlineServicesFactoryEOS() {}
-	virtual TSharedPtr<IOnlineServices> Create(FName InInstanceName) override
+	virtual TSharedPtr<IOnlineServices> Create(FName InInstanceName, FName InInstanceConfigName) override
 	{
-		return MakeShared<FOnlineServicesEOS>(InInstanceName);
+		return MakeShared<FOnlineServicesEOS>(InInstanceName, InInstanceConfigName);
 	}
 };
 
diff --git a/Engine/Plugins/Online/OnlineServicesEOS/Source/Public/Online/OnlineServicesEOS.h b/Engine/Plugins/Online/OnlineServicesEOS/Source/Public/Online/OnlineServicesEOS.h
--- a/Engine/Plugins/Online/OnlineServicesEOS/Source/Public/Online/OnlineServicesEOS.h
+++ b/Engine/Plugins/Online/OnlineServicesEOS/Source/Public/Online/OnlineServicesEOS.h
@@ -22,7 +22,7 @@ class ONLINESERVICESEOS_API FOnlineServicesEOS : public FOnlineServicesEOSGS
 public:
 	using Super = FOnlineServicesEOSGS;
 
-	FOnlineServicesEOS(FName InInstanceName);
+	FOnlineServicesEOS(FName InInstanceName, FName InInstanceConfigName);
 	virtual ~FOnlineServicesEOS() = default;
 
 	virtual void RegisterComponents() override;
diff --git a/Engine/Plugins/Online/OnlineServicesEOSGS/Source/Private/Online/OnlineServicesEOSGS.cpp b/Engine/Plugins/Online/OnlineServicesEOSGS/Source/Private/Online/OnlineServicesEOSGS.cpp
--- a/Engine/Plugins/Online/OnlineServicesEOSGS/Source/Private/Online/OnlineServicesEOSGS.cpp
+++ b/Engine/Plugins/Online/OnlineServicesEOSGS/Source/Private/Online/OnlineServicesEOSGS.cpp
@@ -23,22 +23,15 @@
 
 namespace UE::Online {
 
-FOnlineServicesEOSGS::FOnlineServicesEOSGS(FName InInstanceName)
-	: Super(GetConfigNameStatic(), InInstanceName)
+FOnlineServicesEOSGS::FOnlineServicesEOSGS(FName InInstanceName, FName InInstanceConfigName)
+	: Super(GetServiceConfigNameStatic(), InInstanceName, InInstanceConfigName)
 {
 }
 
 void FOnlineServicesEOSGS::Init()
 {
 	FOnlineServicesEOSGSPlatformFactory& PlatformFactory = FOnlineServicesEOSGSPlatformFactory::Get();
-	if (InstanceName.IsNone())
-	{
-		EOSPlatformHandle = PlatformFactory.GetDefaultPlatform();
-	}
-	else
-	{
-		EOSPlatformHandle = PlatformFactory.CreatePlatform(InstanceName);
-	}
+	EOSPlatformHandle = PlatformFactory.CreatePlatform(InstanceName, InstanceConfigName);
 	if (EOSPlatformHandle)
 	{
 #if WITH_ENGINE
diff --git a/Engine/Plugins/Online/OnlineServicesEOSGS/Source/Private/Online/OnlineServicesEOSGSModule.cpp b/Engine/Plugins/Online/OnlineServicesEOSGS/Source/Private/Online/OnlineServicesEOSGSModule.cpp
--- a/Engine/Plugins/Online/OnlineServicesEOSGS/Source/Private/Online/OnlineServicesEOSGSModule.cpp
+++ b/Engine/Plugins/Online/OnlineServicesEOSGS/Source/Private/Online/OnlineServicesEOSGSModule.cpp
@@ -16,9 +16,9 @@ class FOnlineServicesFactoryEOSGS : public IOnlineServicesFactory
 {
 public:
 	virtual ~FOnlineServicesFactoryEOSGS() {}
-	virtual TSharedPtr<IOnlineServices> Create(FName InInstanceName) override
+	virtual TSharedPtr<IOnlineServices> Create(FName InInstanceName, FName InInstanceConfigName) override
 	{
-		return MakeShared<FOnlineServicesEOSGS>(InInstanceName);
+		return MakeShared<FOnlineServicesEOSGS>(InInstanceName, InInstanceConfigName);
 	}
 };
 
diff --git a/Engine/Plugins/Online/OnlineServicesEOSGS/Source/Private/Online/OnlineServicesEOSGSPlatformFactory.cpp b/Engine/Plugins/Online/OnlineServicesEOSGS/Source/Private/Online/OnlineServicesEOSGSPlatformFactory.cpp
--- a/Engine/Plugins/Online/OnlineServicesEOSGS/Source/Private/Online/OnlineServicesEOSGSPlatformFactory.cpp
+++ b/Engine/Plugins/Online/OnlineServicesEOSGS/Source/Private/Online/OnlineServicesEOSGSPlatformFactory.cpp
@@ -18,7 +18,7 @@
 
 FString LoadEOSPlatformConfig(IEOSSDKManager* const SDKManager)
 {
-	FString ConfigSectionName = FString::Printf(TEXT("OnlineServices.%s"), UE::Online::FOnlineServicesEOSGS::GetConfigNameStatic());
+	FString ConfigSectionName = FString::Printf(TEXT("OnlineServices.%s"), UE::Online::FOnlineServicesEOSGS::GetServiceConfigNameStatic());
 	if (!GConfig->DoesSectionExist(*ConfigSectionName, GEngineIni))
 	{
 		return FString();
@@ -99,7 +99,7 @@ void FOnlineServicesEOSGSPlatformFactory::TearDown()
 	return TLazySingleton<FOnlineServicesEOSGSPlatformFactory>::TearDown();
 }
 
-IEOSPlatformHandlePtr FOnlineServicesEOSGSPlatformFactory::CreatePlatform(FName InstanceName)
+IEOSPlatformHandlePtr FOnlineServicesEOSGSPlatformFactory::CreatePlatform(FName InstanceName, FName InstanceConfigName)
 {
 	const FName EOSSharedModuleName = TEXT("EOSShared");
 	if (!FModuleManager::Get().IsModuleLoaded(EOSSharedModuleName))
@@ -120,7 +120,11 @@ IEOSPlatformHandlePtr FOnlineServicesEOSGSPlatformFactory::CreatePlatform(FName
 		return {};
 	}
 
-	FString PlatformConfigName = LoadEOSPlatformConfig(SDKManager);
+	FString PlatformConfigName = InstanceConfigName.ToString();
+	if (PlatformConfigName.IsEmpty())
+	{
+		PlatformConfigName = LoadEOSPlatformConfig(SDKManager);
+	}
 	if (PlatformConfigName.IsEmpty())
 	{
 		// Check for default platform config that other modules may have setup.
@@ -141,14 +145,4 @@ IEOSPlatformHandlePtr FOnlineServicesEOSGSPlatformFactory::CreatePlatform(FName
 	return EOSPlatformHandle;
 }
 
-IEOSPlatformHandlePtr FOnlineServicesEOSGSPlatformFactory::GetDefaultPlatform()
-{
-	if (!DefaultEOSPlatformHandle)
-	{
-		DefaultEOSPlatformHandle = CreatePlatform(NAME_None);
-	}
-
-	return DefaultEOSPlatformHandle;
-}
-
 /* UE::Online */ }
diff --git a/Engine/Plugins/Online/OnlineServicesEOSGS/Source/Private/Online/OnlineServicesEOSGSPlatformFactory.h b/Engine/Plugins/Online/OnlineServicesEOSGS/Source/Private/Online/OnlineServicesEOSGSPlatformFactory.h
--- a/Engine/Plugins/Online/OnlineServicesEOSGS/Source/Private/Online/OnlineServicesEOSGSPlatformFactory.h
+++ b/Engine/Plugins/Online/OnlineServicesEOSGS/Source/Private/Online/OnlineServicesEOSGSPlatformFactory.h
@@ -25,21 +25,16 @@ public:
 	static void TearDown();
 
 	/**
-	 * Create a new platform. Loads configuration from ini.
+	 * Create a new platform instance for the given Instance/Config pair.
+	 * Note If InstanceConfigName = NAME_None, this will attempt to resolve a config to use from various sources, including OnlineServices config, and EOSSDKManager cached configs.
+	 * @param InstanceName the "instance", typically a WorldContextHandle
+	 * @param InstanceConfigName the named EOS platform config to use
 	 * @return a new platform, or null on failure
 	 */
-	IEOSPlatformHandlePtr CreatePlatform(FName InstanceName);
-	/**
-	 * Get the default platform.
-	 * The default platform is created on startup and can be used by the default instance of online services EOS.
-	 * @return the default platform
-	 */
-	IEOSPlatformHandlePtr GetDefaultPlatform();
+	IEOSPlatformHandlePtr CreatePlatform(FName InstanceName, FName InstanceConfigName);
 private:
 	FOnlineServicesEOSGSPlatformFactory();
 	friend FLazySingleton;
-	/** Platform handle for the default instance */
-	IEOSPlatformHandlePtr DefaultEOSPlatformHandle;
 };
 
 /* UE::Online */ }
diff --git a/Engine/Plugins/Online/OnlineServicesEOSGS/Source/Public/Online/OnlineServicesEOSGS.h b/Engine/Plugins/Online/OnlineServicesEOSGS/Source/Public/Online/OnlineServicesEOSGS.h
--- a/Engine/Plugins/Online/OnlineServicesEOSGS/Source/Public/Online/OnlineServicesEOSGS.h
+++ b/Engine/Plugins/Online/OnlineServicesEOSGS/Source/Public/Online/OnlineServicesEOSGS.h
@@ -22,7 +22,7 @@ class ONLINESERVICESEOSGS_API FOnlineServicesEOSGS : public FOnlineServicesCommo
 public:
 	using Super = FOnlineServicesCommon;
 
-	FOnlineServicesEOSGS(FName InInstanceName);
+	FOnlineServicesEOSGS(FName InstanceName, FName InstanceConfigName);
 	virtual ~FOnlineServicesEOSGS() = default;
 
 	virtual void Init() override;
@@ -33,7 +33,9 @@ public:
 
 	IEOSPlatformHandlePtr GetEOSPlatformHandle() const { return EOSPlatformHandle; }
 
-	static const TCHAR* GetConfigNameStatic() { return TEXT("EOS"); }
+	UE_DEPRECATED(5.5, "Please call GetServiceConfigNameStatic instead")
+	static const TCHAR* GetConfigNameStatic() { return GetServiceConfigNameStatic(); }
+	static const TCHAR* GetServiceConfigNameStatic() { return TEXT("EOS"); }
 protected:
 	IEOSPlatformHandlePtr EOSPlatformHandle;
 
diff --git a/Engine/Plugins/Online/OnlineServicesNull/Source/Private/Online/OnlineServicesNullModule.cpp b/Engine/Plugins/Online/OnlineServicesNull/Source/Private/Online/OnlineServicesNullModule.cpp
--- a/Engine/Plugins/Online/OnlineServicesNull/Source/Private/Online/OnlineServicesNullModule.cpp
+++ b/Engine/Plugins/Online/OnlineServicesNull/Source/Private/Online/OnlineServicesNullModule.cpp
@@ -15,7 +15,7 @@ class FOnlineServicesFactoryNull : public IOnlineServicesFactory
 {
 public:
 	virtual ~FOnlineServicesFactoryNull() {}
-	virtual TSharedPtr<IOnlineServices> Create(FName InInstanceName) override
+	virtual TSharedPtr<IOnlineServices> Create(FName InInstanceName, FName InstanceConfigName) override
 	{
 		return MakeShared<FOnlineServicesNull>(InInstanceName);
 	}
diff --git a/Engine/Plugins/Online/OnlineServicesOSSAdapter/Source/Private/Online/OnlineServicesOSSAdapter.cpp b/Engine/Plugins/Online/OnlineServicesOSSAdapter/Source/Private/Online/OnlineServicesOSSAdapter.cpp
--- a/Engine/Plugins/Online/OnlineServicesOSSAdapter/Source/Private/Online/OnlineServicesOSSAdapter.cpp
+++ b/Engine/Plugins/Online/OnlineServicesOSSAdapter/Source/Private/Online/OnlineServicesOSSAdapter.cpp
@@ -22,8 +22,8 @@
 
 namespace UE::Online {
 
-FOnlineServicesOSSAdapter::FOnlineServicesOSSAdapter(EOnlineServices InServicesType, const FString& InConfigName, FName InInstanceName, IOnlineSubsystem* InSubsystem)
-	: FOnlineServicesCommon(InConfigName, InInstanceName)
+FOnlineServicesOSSAdapter::FOnlineServicesOSSAdapter(EOnlineServices InServicesType, const FString& InServiceConfigName, FName InInstanceName, IOnlineSubsystem* InSubsystem)
+	: FOnlineServicesCommon(InServiceConfigName, InInstanceName)
 	, ServicesType(InServicesType)
 	, Subsystem(InSubsystem)
 {
diff --git a/Engine/Plugins/Online/OnlineServicesOSSAdapter/Source/Private/Online/OnlineServicesOSSAdapterModule.cpp b/Engine/Plugins/Online/OnlineServicesOSSAdapter/Source/Private/Online/OnlineServicesOSSAdapterModule.cpp
--- a/Engine/Plugins/Online/OnlineServicesOSSAdapter/Source/Private/Online/OnlineServicesOSSAdapterModule.cpp
+++ b/Engine/Plugins/Online/OnlineServicesOSSAdapter/Source/Private/Online/OnlineServicesOSSAdapterModule.cpp
@@ -59,7 +59,7 @@ public:
 	}
 
 	virtual ~FOnlineServicesFactoryOSSAdapter() {}
-	virtual TSharedPtr<IOnlineServices> Create(FName InInstanceName) override
+	virtual TSharedPtr<IOnlineServices> Create(FName InInstanceName, FName /*InstanceConfigName*/) override
 	{
 		FName CombinedInstanceName = InInstanceName.IsNone() ? (/*NoSuffix*/Config.OnlineSubsystem) : FName(*FString::Printf(TEXT("%s:%s"), *Config.OnlineSubsystem.ToString(), *InInstanceName.ToString()));
 		IOnlineSubsystem* Subsystem = IOnlineSubsystem::Get(CombinedInstanceName);
diff --git a/Engine/Plugins/Online/OnlineServicesOSSAdapter/Source/Public/Online/OnlineServicesOSSAdapter.h b/Engine/Plugins/Online/OnlineServicesOSSAdapter/Source/Public/Online/OnlineServicesOSSAdapter.h
--- a/Engine/Plugins/Online/OnlineServicesOSSAdapter/Source/Public/Online/OnlineServicesOSSAdapter.h
+++ b/Engine/Plugins/Online/OnlineServicesOSSAdapter/Source/Public/Online/OnlineServicesOSSAdapter.h
@@ -15,7 +15,7 @@ class ONLINESERVICESOSSADAPTER_API FOnlineServicesOSSAdapter : public FOnlineSer
 public:
 	using Super = FOnlineServicesCommon;
 
-	FOnlineServicesOSSAdapter(EOnlineServices InServicesType, const FString& InConfigName, FName InInstanceName, IOnlineSubsystem* InSubsystem);
+	FOnlineServicesOSSAdapter(EOnlineServices InServicesType, const FString& InServiceConfigName, FName InInstanceName, IOnlineSubsystem* InSubsystem);
 
 	virtual void RegisterComponents() override;
 	virtual void Initialize() override;
