Disable FOnlineError from being constructable from a bool
Call FOnlineEventDelegateHandle unbind on move assignment

diff --git a/Engine/Plugins/Online/OnlineServices/Source/OnlineServicesInterface/Public/Online/OnlineAsyncOpHandle.h b/Engine/Plugins/Online/OnlineServices/Source/OnlineServicesInterface/Public/Online/OnlineAsyncOpHandle.h
--- a/Engine/Plugins/Online/OnlineServices/Source/OnlineServicesInterface/Public/Online/OnlineAsyncOpHandle.h
+++ b/Engine/Plugins/Online/OnlineServices/Source/OnlineServicesInterface/Public/Online/OnlineAsyncOpHandle.h
@@ -229,6 +229,7 @@ public:
 
 	FOnlineEventDelegateHandle& operator=(FOnlineEventDelegateHandle&& Other)
 	{
+		Unbind();
 		DelegateUnbinder = MoveTemp(Other.DelegateUnbinder);
 		return *this;
 	}
diff --git a/Engine/Plugins/Online/OnlineServices/Source/OnlineServicesInterface/Public/Online/OnlineError.h b/Engine/Plugins/Online/OnlineServices/Source/OnlineServicesInterface/Public/Online/OnlineError.h
--- a/Engine/Plugins/Online/OnlineServices/Source/OnlineServicesInterface/Public/Online/OnlineError.h
+++ b/Engine/Plugins/Online/OnlineServices/Source/OnlineServicesInterface/Public/Online/OnlineError.h
@@ -59,13 +59,13 @@ protected:
 class FOnlineError
 {
 public:
-	FOnlineError(ErrorCodeType InErrorCode, TSharedPtr<const IOnlineErrorDetails, ESPMode::ThreadSafe> InDetails = nullptr, TSharedPtr<const FOnlineError, ESPMode::ThreadSafe> InInner = nullptr)
+	explicit FOnlineError(ErrorCodeType InErrorCode, TSharedPtr<const IOnlineErrorDetails, ESPMode::ThreadSafe> InDetails = nullptr, TSharedPtr<const FOnlineError, ESPMode::ThreadSafe> InInner = nullptr)
 		: Details(InDetails)
 		, Inner(InInner)
 		, ErrorCode(InErrorCode)
 	{
 	}
-
+	explicit FOnlineError(bool) = delete;
 
 	FText GetText() const
 	{
