Add ConfigInstanceName to UE::Online::IsLoaded to check if a specific online services instance exists

diff --git a/Engine/Plugins/Online/OnlineServices/Source/OnlineServicesInterface/Private/Online/OnlineServices.cpp b/Engine/Plugins/Online/OnlineServices/Source/OnlineServicesInterface/Private/Online/OnlineServices.cpp
--- a/Engine/Plugins/Online/OnlineServices/Source/OnlineServicesInterface/Private/Online/OnlineServices.cpp
+++ b/Engine/Plugins/Online/OnlineServices/Source/OnlineServicesInterface/Private/Online/OnlineServices.cpp
@@ -65,9 +65,9 @@ int32 GetBuildUniqueId()
 	return BuildId;
 }
 
-bool IsLoaded(EOnlineServices OnlineServices, FName InstanceName)
+bool IsLoaded(EOnlineServices OnlineServices, FName InstanceName, FName InstanceConfigName)
 {
-	return FOnlineServicesRegistry::Get().IsLoaded(OnlineServices, InstanceName);
+	return FOnlineServicesRegistry::Get().IsLoaded(OnlineServices, InstanceName, InstanceConfigName);
 }
 
 TSharedPtr<IOnlineServices> GetServices(EOnlineServices OnlineServices, FName InstanceName, FName InstanceConfigName)
diff --git a/Engine/Plugins/Online/OnlineServices/Source/OnlineServicesInterface/Public/Online/OnlineServices.h b/Engine/Plugins/Online/OnlineServices/Source/OnlineServicesInterface/Public/Online/OnlineServices.h
--- a/Engine/Plugins/Online/OnlineServices/Source/OnlineServicesInterface/Public/Online/OnlineServices.h
+++ b/Engine/Plugins/Online/OnlineServices/Source/OnlineServicesInterface/Public/Online/OnlineServices.h
@@ -205,7 +205,7 @@ ONLINESERVICESINTERFACE_API int32 GetBuildUniqueId();
  * @param InstanceName Name of the instance
  * @return true if the instance is loaded
  */
-ONLINESERVICESINTERFACE_API bool IsLoaded(EOnlineServices OnlineServices = EOnlineServices::Default, FName InstanceName = NAME_None);
+ONLINESERVICESINTERFACE_API bool IsLoaded(EOnlineServices OnlineServices = EOnlineServices::Default, FName InstanceName = NAME_None, FName InstanceConfigName = NAME_None);
 
 /**
  * Get an instance of the online service
