- Improve FOnlineServicesCommon to re-add the old constructor sig but marked deprecated, so implementers are directed to integrate the new ConfigName param

diff --git a/Engine/Plugins/Online/OnlineServices/Source/OnlineServicesCommon/Public/Online/OnlineServicesCommon.h b/Engine/Plugins/Online/OnlineServices/Source/OnlineServicesCommon/Public/Online/OnlineServicesCommon.h
--- a/Engine/Plugins/Online/OnlineServices/Source/OnlineServicesCommon/Public/Online/OnlineServicesCommon.h
+++ b/Engine/Plugins/Online/OnlineServices/Source/OnlineServicesCommon/Public/Online/OnlineServicesCommon.h
@@ -23,8 +23,14 @@ class ONLINESERVICESCOMMON_API FOnlineServicesCommon
 {
 public:
 	using Super = IOnlineServices;
+	
+	UE_DEPRECATED(5.5, "Please call the new method taking an FName InstanceConfigName")
+	FOnlineServicesCommon(const FString& InServiceConfigName, FName InInstanceName)
+		: FOnlineServicesCommon(InServiceConfigName, InInstanceName, NAME_None)
+	{
+	}
 
-	FOnlineServicesCommon(const FString& InServiceConfigName, FName InInstanceName, FName InInstanceConfigName = NAME_None);
+	FOnlineServicesCommon(const FString& InServiceConfigName, FName InInstanceName, FName InInstanceConfigName);
 	FOnlineServicesCommon(const FOnlineServicesCommon&) = delete;
 	FOnlineServicesCommon(FOnlineServicesCommon&&) = delete;
 	virtual ~FOnlineServicesCommon();
diff --git a/Engine/Plugins/Online/OnlineServicesOSSAdapter/Source/Private/Online/OnlineServicesOSSAdapter.cpp b/Engine/Plugins/Online/OnlineServicesOSSAdapter/Source/Private/Online/OnlineServicesOSSAdapter.cpp
--- a/Engine/Plugins/Online/OnlineServicesOSSAdapter/Source/Private/Online/OnlineServicesOSSAdapter.cpp
+++ b/Engine/Plugins/Online/OnlineServicesOSSAdapter/Source/Private/Online/OnlineServicesOSSAdapter.cpp
@@ -23,7 +23,7 @@
 namespace UE::Online {
 
 FOnlineServicesOSSAdapter::FOnlineServicesOSSAdapter(EOnlineServices InServicesType, const FString& InServiceConfigName, FName InInstanceName, IOnlineSubsystem* InSubsystem)
-	: FOnlineServicesCommon(InServiceConfigName, InInstanceName)
+	: FOnlineServicesCommon(InServiceConfigName, InInstanceName, NAME_None)
 	, ServicesType(InServicesType)
 	, Subsystem(InSubsystem)
 {
