Add mutable Get to FOnlineOperationData

diff --git a/Engine/Plugins/Online/OnlineServices/Source/OnlineServicesCommon/Public/Online/OnlineAsyncOp.h b/Engine/Plugins/Online/OnlineServices/Source/OnlineServicesCommon/Public/Online/OnlineAsyncOp.h
--- a/Engine/Plugins/Online/OnlineServices/Source/OnlineServicesCommon/Public/Online/OnlineAsyncOp.h
+++ b/Engine/Plugins/Online/OnlineServices/Source/OnlineServicesCommon/Public/Online/OnlineAsyncOp.h
@@ -113,6 +113,17 @@ public:
 		return nullptr;
 	}
 
+	template <typename T>
+	T* Get(const FString& Key)
+	{
+		if (auto Value = Data.Find(FOperationDataKey{ TOnlineTypeInfo<T>::GetTypeName(), Key }))
+		{
+			return static_cast<T*>((*Value)->GetData());
+		}
+
+		return nullptr;
+	}
+
 	struct FOperationDataKey
 	{
 		FOnlineTypeName TypeName;
