Fix for missing field in OSSv2 presence update.

diff --git a/Engine/Plugins/Online/OnlineServices/Source/OnlineServicesInterface/Public/Online/Presence.h b/Engine/Plugins/Online/OnlineServices/Source/OnlineServicesInterface/Public/Online/Presence.h
--- a/Engine/Plugins/Online/OnlineServices/Source/OnlineServicesInterface/Public/Online/Presence.h
+++ b/Engine/Plugins/Online/OnlineServices/Source/OnlineServicesInterface/Public/Online/Presence.h
@@ -94,7 +94,7 @@ struct FUserPresence
 	FString StatusString;
 	/** Game-defined representation of the current game state (e.g. "Squads- 4 Remaining" or "Level 5 Warrior") */
 	FString RichPresenceString;
-	/** Session keys */
+	/** Session keys, and any platform specific presence properties */
 	FPresenceProperties Properties;
 };
 	
@@ -415,6 +415,7 @@ BEGIN_ONLINE_STRUCT_META(FPartialUpdatePresence::Params::FMutations)
 	ONLINE_STRUCT_FIELD(FPartialUpdatePresence::Params::FMutations, Joinability),
 	ONLINE_STRUCT_FIELD(FPartialUpdatePresence::Params::FMutations, GameStatus),
 	ONLINE_STRUCT_FIELD(FPartialUpdatePresence::Params::FMutations, StatusString),
+	ONLINE_STRUCT_FIELD(FPartialUpdatePresence::Params::FMutations, RichPresenceString),
 	ONLINE_STRUCT_FIELD(FPartialUpdatePresence::Params::FMutations, UpdatedProperties),
 	ONLINE_STRUCT_FIELD(FPartialUpdatePresence::Params::FMutations, RemovedProperties)
 END_ONLINE_STRUCT_META()
