Added explicit validity check for OnlineServices EOSGS Lobby notification logic.

diff --git a/Engine/Plugins/Online/OnlineServices/Source/OnlineServicesCommon/Private/Online/LobbiesCommonTypes.cpp b/Engine/Plugins/Online/OnlineServices/Source/OnlineServicesCommon/Private/Online/LobbiesCommonTypes.cpp
--- a/Engine/Plugins/Online/OnlineServices/Source/OnlineServicesCommon/Private/Online/LobbiesCommonTypes.cpp
+++ b/Engine/Plugins/Online/OnlineServices/Source/OnlineServicesCommon/Private/Online/LobbiesCommonTypes.cpp
@@ -240,12 +240,20 @@ FLobbyClientDataCommitServiceSnapshot::Result FLobbyClientData::CommitServiceSna
 	// Check for ownership change.
 	if (PreparedServiceChanges->OwnerAccountId)
 	{
-		InternalPublicData->OwnerAccountId = *PreparedServiceChanges->OwnerAccountId;
-
-		if (bDispatchNotifications)
+		// We check validity of the member data because at the time when the snapshot is retrieved containing a host change, that new host might have already left the session
+		if (TSharedRef<FLobbyMemberInternal>* OwnerMemberData = MemberDataStorage.Find(*PreparedServiceChanges->OwnerAccountId))
 		{
-			Params.LobbyEvents->OnLobbyLeaderChanged.Broadcast(FLobbyLeaderChanged{ InternalPublicData, MemberDataStorage.FindChecked(InternalPublicData->OwnerAccountId) });
+			InternalPublicData->OwnerAccountId = *PreparedServiceChanges->OwnerAccountId;
+
+			if (bDispatchNotifications)
+			{
+				Params.LobbyEvents->OnLobbyLeaderChanged.Broadcast(FLobbyLeaderChanged{ InternalPublicData, *OwnerMemberData });
+			}
 		}
+		else
+		{
+			InternalPublicData->OwnerAccountId = FAccountId();
+		}		
 	}
 
 	// Check for max member change.
