Lobbies: Fix issue setting a lobby member attribute when joining a lobby.

diff --git a/Engine/Plugins/Online/OnlineServices/Source/OnlineServicesCommon/Private/Online/LobbiesCommonTests.cpp b/Engine/Plugins/Online/OnlineServices/Source/OnlineServicesCommon/Private/Online/LobbiesCommonTests.cpp
--- a/Engine/Plugins/Online/OnlineServices/Source/OnlineServicesCommon/Private/Online/LobbiesCommonTests.cpp
+++ b/Engine/Plugins/Online/OnlineServices/Source/OnlineServicesCommon/Private/Online/LobbiesCommonTests.cpp
@@ -1201,7 +1201,7 @@ bool FClientLobbyDataTest::RunTest(const FString& Parameters)
 		}
 	}
 
-	// 5. Local changes supplemental
+	// 5. Local changes supplemental - Local user kicks remote user.
 	{
 		FLobbyEvents LobbyEvents;
 		FLobbyEventCapture EventCapture(LobbyEvents);
@@ -1274,6 +1274,73 @@ bool FClientLobbyDataTest::RunTest(const FString& Parameters)
 		UTEST_EQUAL("Check LobbyMemberLeft valid.", EventCapture.LobbyMemberLeft[0].Notification.Reason, ELobbyMemberLeaveReason::Kicked);
 	}
 
+	// 6. Local changes supplemental - Local user joins a lobby, joining user is not included in lobby data snapshot.
+	{
+		FLobbyEvents LobbyEvents;
+		FLobbyEventCapture EventCapture(LobbyEvents);
+		FLobbyId LobbyId(EOnlineServices::Null, 1);
+		FLobbyClientData ClientData(LobbyId, SchemaRegistry);
+
+		FAccountId User1(EOnlineServices::Null, 1);
+		FAccountId User2(EOnlineServices::Null, 2);
+		int32 LobbyMaxMembers = 5;
+		ELobbyJoinPolicy LobbyJoinPolicy = ELobbyJoinPolicy::PublicAdvertised;
+
+		// 6.1 setup. Create a lobby from a snapshot containing one user.
+		{
+			TMap<FAccountId, FLobbyMemberServiceSnapshot> LobbyMemberSnapshots;
+			{
+				FLobbyMemberServiceSnapshot& MemberSnapshot = LobbyMemberSnapshots.Add(User1);
+				MemberSnapshot.AccountId = User1;
+				MemberSnapshot.SchemaServiceSnapshot.Attributes.Add(LobbyMemberServiceAttributeId1, LobbyMemberAttribute1Value1);
+			}
+
+			FLobbyServiceSnapshot LobbySnapshot;
+			LobbySnapshot.OwnerAccountId = User1;
+			LobbySnapshot.MaxMembers = LobbyMaxMembers;
+			LobbySnapshot.JoinPolicy = LobbyJoinPolicy;
+			LobbySnapshot.SchemaServiceSnapshot.Attributes.Add(LobbyServiceAttributeId1, SchemaCompatibilityAttributeValue);
+			LobbySnapshot.SchemaServiceSnapshot.Attributes.Add(LobbyServiceAttributeId2, LobbyAttribute1Value1);
+			LobbySnapshot.Members = {User1};
+
+			TMap<FAccountId, ELobbyMemberLeaveReason> LeaveReasons;
+			TOnlineResult<FLobbyClientDataPrepareServiceSnapshot> PrepareResult = ClientData.PrepareServiceSnapshot({ MoveTemp(LobbySnapshot), MoveTemp(LobbyMemberSnapshots), MoveTemp(LeaveReasons) });
+			UTEST_TRUE("Check PrepareServiceSnapshot succeeded.", PrepareResult.IsOk());
+
+			FLobbyClientDataCommitServiceSnapshot::Result CommitResult = ClientData.CommitServiceSnapshot({ &LobbyEvents });
+			UTEST_EQUAL("Check no local members left.", CommitResult.LeavingLocalMembers.Num(), 0);
+		}
+
+		// 6.2 Join local user to the lobby and set attribute in derived schema.
+		{
+			FLobbyClientChanges LobbyChanges;
+			LobbyChanges.MemberAttributes = { {{ LobbyMemberAttributeId1, LobbyMemberAttribute1Value2 }}, {} };
+			TOnlineResult<FLobbyClientDataPrepareClientChanges> PrepareResult = ClientData.PrepareClientChanges({ User2, MoveTemp(LobbyChanges) });
+			UTEST_TRUE("Check PrepareClientChanges succeeded.", PrepareResult.IsOk());
+
+			FLobbyClientDataCommitClientChanges::Result CommitResult = ClientData.CommitClientChanges({ &LobbyEvents });
+			UTEST_EQUAL("Check no local members left.", CommitResult.LeavingLocalMembers.Num(), 0);
+		}
+
+		// 6.3 Test lobby events received.
+		UTEST_EQUAL("Check events received.", EventCapture.GetTotalNotificationsReceived(), 3);
+		UTEST_EQUAL("Check LobbyJoined valid.", EventCapture.LobbyJoined.Num(), 1);
+		UTEST_EQUAL("Check LobbyJoined valid.", EventCapture.LobbyJoined[0].GlobalIndex, 0);
+		UTEST_EQUAL("Check LobbyMemberJoined[0] valid.", EventCapture.LobbyMemberJoined.Num(), 2);
+		UTEST_EQUAL("Check LobbyMemberJoined[0] valid.", EventCapture.LobbyMemberJoined[0].GlobalIndex, 1);
+		UTEST_EQUAL("Check LobbyMemberJoined[0] valid.", EventCapture.LobbyMemberJoined[0].Notification.Member->AccountId, User1);
+		UTEST_EQUAL("Check LobbyMemberJoined[0] attributes are valid.", EventCapture.LobbyMemberJoined[0].Notification.Member->Attributes.Num(), 1);
+		const FSchemaVariant* LobbyMemberAttribute1 = EventCapture.LobbyMemberJoined[0].Notification.Member->Attributes.Find(LobbyMemberAttributeId1);
+		UTEST_NOT_NULL("Check LobbyMemberJoined[0] attributes are valid.", LobbyMemberAttribute1);
+		UTEST_EQUAL("Check LobbyMemberJoined[0] attributes are valid.", *LobbyMemberAttribute1, LobbyMemberAttribute1Value1);
+		UTEST_EQUAL("Check LobbyMemberJoined[1] valid.", EventCapture.LobbyMemberJoined[1].GlobalIndex, 2);
+		UTEST_EQUAL("Check LobbyMemberJoined[1] valid.", EventCapture.LobbyMemberJoined[1].Notification.Member->AccountId, User2);
+		UTEST_EQUAL("Check LobbyMemberJoined[1] attributes are valid.", EventCapture.LobbyMemberJoined[1].Notification.Member->Attributes.Num(), 1);
+		const FSchemaVariant* LobbyMemberAttribute2 = EventCapture.LobbyMemberJoined[1].Notification.Member->Attributes.Find(LobbyMemberAttributeId1);
+		UTEST_NOT_NULL("Check LobbyMemberJoined[1] attributes are valid.", LobbyMemberAttribute2);
+		UTEST_EQUAL("Check LobbyMemberJoined[1] attributes are valid.", *LobbyMemberAttribute2, LobbyMemberAttribute1Value2);
+	}
+
 	return true;
 }
 
diff --git a/Engine/Plugins/Online/OnlineServices/Source/OnlineServicesCommon/Private/Online/LobbiesCommonTypes.cpp b/Engine/Plugins/Online/OnlineServices/Source/OnlineServicesCommon/Private/Online/LobbiesCommonTypes.cpp
--- a/Engine/Plugins/Online/OnlineServices/Source/OnlineServicesCommon/Private/Online/LobbiesCommonTypes.cpp
+++ b/Engine/Plugins/Online/OnlineServices/Source/OnlineServicesCommon/Private/Online/LobbiesCommonTypes.cpp
@@ -443,8 +443,8 @@ TOnlineResult<FLobbyClientDataPrepareClientChanges> FLobbyClientData::PrepareCli
 			}
 			else
 			{
-				// New member, create an entry for them.
-				TSharedRef<FLobbyMemberInternal> NewMemberData = MakeShared<FLobbyMemberInternal>(SchemaRegistry);
+				// New member who does not have a service snapshot, create an entry for them using the current lobby schema.
+				TSharedRef<FLobbyMemberInternal> NewMemberData = MakeShared<FLobbyMemberInternal>(SchemaRegistry, LobbySchemaId);
 				NewMemberData->AccountId = Params.LocalAccountId;
 				NewPreparedClientChanges.JoiningLocalMember = NewMemberData;
 				PreparingMemberData = &NewMemberData.Get();
@@ -637,7 +637,7 @@ FLobbyClientDataCommitClientChanges::Result FLobbyClientData::CommitClientChange
 			AddMember(JoinedMember);
 
 			// Commit attribute changes.
-			JoinedMember->SchemaData.CommitServiceSnapshot();
+			JoinedMember->SchemaData.CommitClientChanges();
 
 			// Send join notification for the local user.
 			if (bDispatchNotifications && !bWasAlreadyJoined)
diff --git a/Engine/Plugins/Online/OnlineServices/Source/OnlineServicesCommon/Public/Online/LobbiesCommonTypes.h b/Engine/Plugins/Online/OnlineServices/Source/OnlineServicesCommon/Public/Online/LobbiesCommonTypes.h
--- a/Engine/Plugins/Online/OnlineServices/Source/OnlineServicesCommon/Public/Online/LobbiesCommonTypes.h
+++ b/Engine/Plugins/Online/OnlineServices/Source/OnlineServicesCommon/Public/Online/LobbiesCommonTypes.h
@@ -113,9 +113,9 @@ struct FLobbyInternal : public FLobby
 */
 struct FLobbyMemberInternal : public FLobbyMember
 {
-	FLobbyMemberInternal(const TSharedRef<const FSchemaRegistry>& SchemaRegistry)
+	FLobbyMemberInternal(const TSharedRef<const FSchemaRegistry>& SchemaRegistry, const FSchemaId& DerivedSchemaId = FSchemaId())
 		: FLobbyMember()
-		, SchemaData(FSchemaId(), LobbyBaseSchemaId, LobbyMemberSchemaCategoryId, SchemaRegistry, Attributes)
+		, SchemaData(DerivedSchemaId, LobbyBaseSchemaId, LobbyMemberSchemaCategoryId, SchemaRegistry, Attributes)
 	{
 	}
 
