Make RegisterComponent function support to pass in the pure virtual function of base class instead of needing the implemment in the common interface.

diff --git a/Engine/Plugins/Online/OnlineServices/Source/OnlineServicesCommon/Public/Online/OnlineComponent.h b/Engine/Plugins/Online/OnlineServices/Source/OnlineServicesCommon/Public/Online/OnlineComponent.h
--- a/Engine/Plugins/Online/OnlineServices/Source/OnlineServicesCommon/Public/Online/OnlineComponent.h
+++ b/Engine/Plugins/Online/OnlineServices/Source/OnlineServicesCommon/Public/Online/OnlineComponent.h
@@ -123,7 +123,14 @@ public:
 		using InterfaceType = typename Private::TOnlineInterfaceOperationMemberFunctionPtrTraits<T>::InterfaceType;
 		static_assert(std::is_base_of_v<ComponentType, InterfaceType> || std::is_same_v<ComponentType, InterfaceType>);
 		
-		TOnlineInterfaceOperationExecHandler<T>* Handler = new TOnlineInterfaceOperationExecHandler<T>(static_cast<InterfaceType*>(this), MemberFunction);
+		RegisterCommand(static_cast<InterfaceType*>(this), MemberFunction);
+	}
+
+	// Default handler: generic parsing of Params, log Result
+	template <typename InterfaceType, typename T>
+	void RegisterCommand(InterfaceType* Interface, T MemberFunction)
+	{
+		TOnlineInterfaceOperationExecHandler<InterfaceType, T>* Handler = new TOnlineInterfaceOperationExecHandler<InterfaceType, T>(Interface, MemberFunction);
 
 		const TCHAR* Name = Private::TOnlineInterfaceOperationMemberFunctionPtrTraits<T>::OpType::Name;
 		ExecCommands.Emplace(Name, Handler);
diff --git a/Engine/Plugins/Online/OnlineServices/Source/OnlineServicesCommon/Public/Online/OnlineExecHandler.h b/Engine/Plugins/Online/OnlineServices/Source/OnlineServicesCommon/Public/Online/OnlineExecHandler.h
--- a/Engine/Plugins/Online/OnlineServices/Source/OnlineServicesCommon/Public/Online/OnlineExecHandler.h
+++ b/Engine/Plugins/Online/OnlineServices/Source/OnlineServicesCommon/Public/Online/OnlineExecHandler.h
@@ -889,11 +889,10 @@ std::enable_if_t<TModels_V<Meta::COnlineMetadataAvailable, T>, bool> ParseOnline
 
 /* Private */ }
 
-template <typename MemberFunctionPtrType>
+template <typename InterfaceType, typename MemberFunctionPtrType>
 class TOnlineInterfaceOperationExecHandler : public IOnlineExecHandler
 {
 public:
-	using InterfaceType = typename Private::TOnlineInterfaceOperationMemberFunctionPtrTraits<MemberFunctionPtrType>::InterfaceType;
 	using OpType = typename Private::TOnlineInterfaceOperationMemberFunctionPtrTraits<MemberFunctionPtrType>::OpType;
 
 	TOnlineInterfaceOperationExecHandler(InterfaceType* InInterface, MemberFunctionPtrType InFunction)
