Schema: Log schema definition when parsed for debugging.
Schema: Allow update to continue when an unexpected attribute is received from a service snapshot.

diff --git a/Engine/Plugins/Online/OnlineServices/Source/OnlineServicesInterface/Private/Online/Schema.cpp b/Engine/Plugins/Online/OnlineServices/Source/OnlineServicesInterface/Private/Online/Schema.cpp
--- a/Engine/Plugins/Online/OnlineServices/Source/OnlineServicesInterface/Private/Online/Schema.cpp
+++ b/Engine/Plugins/Online/OnlineServices/Source/OnlineServicesInterface/Private/Online/Schema.cpp
@@ -5,6 +5,7 @@
 #include "Algo/Find.h"
 #include "Algo/ForEach.h"
 #include "Algo/Transform.h"
+#include "Misc/EnumRange.h"
 #include "Online/OnlineErrorDefinitions.h"
 #include "Online/OnlineResult.h"
 
@@ -12,6 +13,9 @@ DEFINE_LOG_CATEGORY_STATIC(LogOnlineSchema, Log, All);
 
 namespace UE::Online {
 
+namespace
+{
+
 template <typename T>
 T BuildFlagsFromArray(const TArray<T>& Input)
 {
@@ -22,7 +26,7 @@ T BuildFlagsFromArray(const TArray<T>& Input)
 
 static uint32 SchemaNameCrc32(const FName& Name, uint32 CRC=0)
 {
-	const FString NameString = Name.ToString().ToLower();
+	const FString NameString = Name.ToString();
 	return FCrc::StrCrc32<TCHAR>(*NameString, CRC);
 }
 
@@ -55,6 +59,60 @@ ESchemaServiceAttributeFlags TranslateAttributeFlags(ESchemaAttributeFlags InFla
 	return OutFlags;
 }
 
+void LogSchemaDefinition(const FSchemaDefinition& Definition)
+{
+	if (UE_LOG_ACTIVE(LogOnlineSchema, Verbose))
+	{
+		auto GetParentString = [](const FSchemaDefinition& Definition) -> FString {
+			TArray<FString> ParentNames;
+			Algo::Transform(Definition.ParentSchemaIds, ParentNames, [](const FSchemaId& ParentSchemaId){ return ParentSchemaId.ToString(); });
+			return FString::Join(ParentNames, TEXT(","));
+		};
+
+		auto GetCategoryString = [](const FSchemaDefinition& Definition) -> FString {
+			TArray<FString> CategoryNames;
+			Algo::Transform(Definition.Categories, CategoryNames, [](const TPair<FSchemaCategoryId, FSchemaCategoryDefinition>& CategoryPair){ return CategoryPair.Key.ToString(); });
+			return FString::Join(CategoryNames, TEXT(","));
+		};
+
+		auto GetFlagsString = [](ESchemaAttributeFlags Flags) -> FString {
+			return FString::JoinBy(MakeFlagsRange(Flags), TEXT("|"), [](ESchemaAttributeFlags Flag){ return LexToString(Flag); });
+		};
+
+		const FString SchemaName = Definition.Id.ToString();
+		const FString ParentSchema = GetParentString(Definition);
+		UE_LOG(LogOnlineSchema, Verbose, TEXT("[SchemaDefinition] SchemaName: [%s]"), *SchemaName);
+		UE_LOG(LogOnlineSchema, Verbose, TEXT("[SchemaDefinition][%s] CompatabilityId: [0x%08" INT64_X_FMT "]"), *SchemaName, Definition.CompatibilityId);
+		if (!ParentSchema.IsEmpty())
+		{
+			UE_LOG(LogOnlineSchema, Verbose, TEXT("[SchemaDefinition][%s] Parent: [%s]"), *SchemaName, *ParentSchema);
+		}
+		UE_LOG(LogOnlineSchema, Verbose, TEXT("[SchemaDefinition][%s] Categories: [%s]"), *SchemaName, *GetCategoryString(Definition));
+		for (const TPair<FSchemaCategoryId, FSchemaCategoryDefinition>& CategoryPair : Definition.Categories)
+		{
+			const FSchemaCategoryDefinition& Category = CategoryPair.Value;
+			const FString CategoryName = Category.Id.ToString();
+			if (Category.SchemaCompatibilityAttributeId != FSchemaId())
+			{
+				UE_LOG(LogOnlineSchema, Verbose, TEXT("[SchemaDefinition][%s][%s] SchemaCompatibilityAttributeId: [%s]"), *SchemaName, *CategoryName, *Category.SchemaCompatibilityAttributeId.ToString());
+			}
+			if (Category.SchemaCompatibilityServiceAttributeId != FSchemaId())
+			{
+				UE_LOG(LogOnlineSchema, Verbose, TEXT("[SchemaDefinition][%s][%s] SchemaCompatibilityServiceAttributeId: [%s]"), *SchemaName, *CategoryName, *Category.SchemaCompatibilityServiceAttributeId.ToString());
+			}
+
+			for (const TPair<FSchemaAttributeId, FSchemaAttributeDefinition>& AttributeDefinitionPair : Category.SchemaAttributeDefinitions)
+			{
+				const FSchemaAttributeDefinition& AttributeDefinition = AttributeDefinitionPair.Value;
+				UE_LOG(LogOnlineSchema, Verbose, TEXT("[SchemaDefinition][%s][%s][%s] Type: [%s], Flags: [%s], MaxSize: [%d], ServiceAttributeId: [%s]"),
+				*SchemaName, *CategoryName, *AttributeDefinition.Id.ToString(), LexToString(AttributeDefinition.Type), *GetFlagsString(AttributeDefinition.Flags), AttributeDefinition.MaxSize, *AttributeDefinition.ServiceAttributeId.ToString());
+			}
+		}
+	}
+}
+
+/* anonymous */ }
+
 bool FSchemaRegistry::ParseConfig(const FSchemaRegistryDescriptorConfig& Config)
 {
 	bool ParsedSuccessfully = true;
@@ -66,7 +124,7 @@ bool FSchemaRegistry::ParseConfig(const FSchemaRegistryDescriptorConfig& Config)
 		// Check that schema descriptor id has not already been used.
 		if (KnownSchemaServiceAttributeDescriptors.Contains(SchemaServiceAttributeDescriptor.Id))
 		{
-			UE_LOG(LogOnlineSchema, Error, TEXT("Duplicate schema service attribute descriptor found: %s"), *SchemaServiceAttributeDescriptor.Id.ToString().ToLower());
+			UE_LOG(LogOnlineSchema, Error, TEXT("Duplicate schema service attribute descriptor found: [%s]"), *SchemaServiceAttributeDescriptor.Id.ToString());
 			ParsedSuccessfully = false;
 			continue;
 		}
@@ -81,7 +139,7 @@ bool FSchemaRegistry::ParseConfig(const FSchemaRegistryDescriptorConfig& Config)
 		// Check that schema service descriptor id has not already been used.
 		if (KnownSchemaServiceDescriptors.Contains(SchemaServiceDescriptor.Id))
 		{
-			UE_LOG(LogOnlineSchema, Error, TEXT("Duplicate schema service descriptor found: %s"), *SchemaServiceDescriptor.Id.ToString().ToLower());
+			UE_LOG(LogOnlineSchema, Error, TEXT("Duplicate schema service descriptor found: [%s]"), *SchemaServiceDescriptor.Id.ToString());
 			ParsedSuccessfully = false;
 			continue;
 		}
@@ -94,8 +152,8 @@ bool FSchemaRegistry::ParseConfig(const FSchemaRegistryDescriptorConfig& Config)
 			if (SeenAttributes.Contains(SchemaServiceAttributeId))
 			{
 				UE_LOG(LogOnlineSchema, Error, TEXT("Invalid schema service attribute %s.%s: Attribute id has already been used."),
-					*SchemaServiceDescriptor.Id.ToString().ToLower(),
-					*SchemaServiceAttributeId.ToString().ToLower());
+					*SchemaServiceDescriptor.Id.ToString(),
+					*SchemaServiceAttributeId.ToString());
 				ParsedSuccessfully = false;
 				continue;
 			}
@@ -105,8 +163,8 @@ bool FSchemaRegistry::ParseConfig(const FSchemaRegistryDescriptorConfig& Config)
 			if (!SchemaServiceAttributeDescriptorPtr)
 			{
 				UE_LOG(LogOnlineSchema, Error, TEXT("Invalid schema service attribute %s.%s: Attribute does not exist."),
-					*SchemaServiceDescriptor.Id.ToString().ToLower(),
-					*SchemaServiceAttributeId.ToString().ToLower());
+					*SchemaServiceDescriptor.Id.ToString(),
+					*SchemaServiceAttributeId.ToString());
 				ParsedSuccessfully = false;
 				continue;
 			}
@@ -117,8 +175,8 @@ bool FSchemaRegistry::ParseConfig(const FSchemaRegistryDescriptorConfig& Config)
 			if (SupportedTypes == ESchemaServiceAttributeSupportedTypeFlags::None)
 			{
 				UE_LOG(LogOnlineSchema, Error, TEXT("Invalid schema service attribute %s.%s: A valid supported type must be selected."),
-					*SchemaServiceDescriptor.Id.ToString().ToLower(),
-					*SchemaServiceAttributeDescriptor.Id.ToString().ToLower());
+					*SchemaServiceDescriptor.Id.ToString(),
+					*SchemaServiceAttributeDescriptor.Id.ToString());
 				ParsedSuccessfully = false;
 			}
 
@@ -126,8 +184,8 @@ bool FSchemaRegistry::ParseConfig(const FSchemaRegistryDescriptorConfig& Config)
 			if (EnumHasAnyFlags(SupportedTypes, ESchemaServiceAttributeSupportedTypeFlags::String) && SchemaServiceAttributeDescriptor.MaxSize <= 0)
 			{
 				UE_LOG(LogOnlineSchema, Error, TEXT("Invalid schema service attribute %s.%s: A valid max size must be set for variable length data."),
-					*SchemaServiceDescriptor.Id.ToString().ToLower(),
-					*SchemaServiceAttributeDescriptor.Id.ToString().ToLower());
+					*SchemaServiceDescriptor.Id.ToString(),
+					*SchemaServiceAttributeDescriptor.Id.ToString());
 				ParsedSuccessfully = false;
 			}
 
@@ -136,8 +194,8 @@ bool FSchemaRegistry::ParseConfig(const FSchemaRegistryDescriptorConfig& Config)
 			if (!EnumHasAnyFlags(ServiceAttributeFlags, ESchemaServiceAttributeFlags::Public | ESchemaServiceAttributeFlags::Private))
 			{
 				UE_LOG(LogOnlineSchema, Error, TEXT("Invalid schema service attribute %s.%s: Either Public or Private visibility must be set."),
-					*SchemaServiceDescriptor.Id.ToString().ToLower(),
-					*SchemaServiceAttributeDescriptor.Id.ToString().ToLower());
+					*SchemaServiceDescriptor.Id.ToString(),
+					*SchemaServiceAttributeDescriptor.Id.ToString());
 				ParsedSuccessfully = false;
 			}
 		}
@@ -152,7 +210,7 @@ bool FSchemaRegistry::ParseConfig(const FSchemaRegistryDescriptorConfig& Config)
 		// Check that schema descriptor id has not already been used.
 		if (KnownSchemaCategoryDescriptors.Contains(SchemaCategoryDescriptor.Id))
 		{
-			UE_LOG(LogOnlineSchema, Error, TEXT("Duplicate schema category descriptor found: %s"), *SchemaCategoryDescriptor.Id.ToString().ToLower());
+			UE_LOG(LogOnlineSchema, Error, TEXT("Duplicate schema category descriptor found: %s"), *SchemaCategoryDescriptor.Id.ToString());
 			ParsedSuccessfully = false;
 			continue;
 		}
@@ -163,8 +221,8 @@ bool FSchemaRegistry::ParseConfig(const FSchemaRegistryDescriptorConfig& Config)
 			&& !KnownSchemaServiceDescriptors.Contains(SchemaCategoryDescriptor.ServiceDescriptorId))
 		{
 			UE_LOG(LogOnlineSchema, Error, TEXT("Invalid schema category %s: Service descriptor id %s not found."),
-				*SchemaCategoryDescriptor.Id.ToString().ToLower(),
-				*SchemaCategoryDescriptor.ServiceDescriptorId.ToString().ToLower());
+				*SchemaCategoryDescriptor.Id.ToString(),
+				*SchemaCategoryDescriptor.ServiceDescriptorId.ToString());
 			ParsedSuccessfully = false;
 		}
 	}
@@ -176,7 +234,7 @@ bool FSchemaRegistry::ParseConfig(const FSchemaRegistryDescriptorConfig& Config)
 		// Check that schema descriptor id has not already been used.
 		if (KnownSchemaDescriptors.Contains(SchemaDescriptor.Id))
 		{
-			UE_LOG(LogOnlineSchema, Error, TEXT("Duplicate schema descriptor found: %s"), *SchemaDescriptor.Id.ToString().ToLower());
+			UE_LOG(LogOnlineSchema, Error, TEXT("Duplicate schema descriptor found: %s"), *SchemaDescriptor.Id.ToString());
 			ParsedSuccessfully = false;
 			continue;
 		}
@@ -187,7 +245,7 @@ bool FSchemaRegistry::ParseConfig(const FSchemaRegistryDescriptorConfig& Config)
 		{
 			// Error, base type must define at least one category
 			UE_LOG(LogOnlineSchema, Error, TEXT("Invalid schema %s: Base schema must define at least one category."),
-				*SchemaDescriptor.Id.ToString().ToLower());
+				*SchemaDescriptor.Id.ToString());
 			ParsedSuccessfully = false;
 		}
 
@@ -196,7 +254,7 @@ bool FSchemaRegistry::ParseConfig(const FSchemaRegistryDescriptorConfig& Config)
 		{
 			// Error child cannot define categories
 			UE_LOG(LogOnlineSchema, Error, TEXT("Invalid schema %s: Child schema cannot define categories."),
-				*SchemaDescriptor.Id.ToString().ToLower());
+				*SchemaDescriptor.Id.ToString());
 			ParsedSuccessfully = false;
 		}
 
@@ -206,8 +264,8 @@ bool FSchemaRegistry::ParseConfig(const FSchemaRegistryDescriptorConfig& Config)
 			if (SeenCategoryIds.Contains(CategoryId))
 			{
 				UE_LOG(LogOnlineSchema, Error, TEXT("Invalid schema category %s.%s: Category id has already been used."),
-					*SchemaDescriptor.Id.ToString().ToLower(),
-					*CategoryId.ToString().ToLower());
+					*SchemaDescriptor.Id.ToString(),
+					*CategoryId.ToString());
 				ParsedSuccessfully = false;
 				continue;
 			}
@@ -216,8 +274,8 @@ bool FSchemaRegistry::ParseConfig(const FSchemaRegistryDescriptorConfig& Config)
 			if (!KnownSchemaCategoryDescriptors.Contains(CategoryId))
 			{
 				UE_LOG(LogOnlineSchema, Error, TEXT("Invalid schema category %s.%s: Category does not exist."),
-					*SchemaDescriptor.Id.ToString().ToLower(),
-					*CategoryId.ToString().ToLower());
+					*SchemaDescriptor.Id.ToString(),
+					*CategoryId.ToString());
 				ParsedSuccessfully = false;
 			}
 		}
@@ -230,7 +288,7 @@ bool FSchemaRegistry::ParseConfig(const FSchemaRegistryDescriptorConfig& Config)
 		// Check that schema descriptor id has not already been used.
 		if (KnownSchemaAttributeDescriptors.Contains(SchemaAttributeDescriptor.Id))
 		{
-			UE_LOG(LogOnlineSchema, Error, TEXT("Duplicate schema attribute descriptor found: %s"), *SchemaAttributeDescriptor.Id.ToString().ToLower());
+			UE_LOG(LogOnlineSchema, Error, TEXT("Duplicate schema attribute descriptor found: %s"), *SchemaAttributeDescriptor.Id.ToString());
 			ParsedSuccessfully = false;
 			continue;
 		}
@@ -240,7 +298,7 @@ bool FSchemaRegistry::ParseConfig(const FSchemaRegistryDescriptorConfig& Config)
 		if (SchemaAttributeDescriptor.Type == ESchemaAttributeType::None)
 		{
 			UE_LOG(LogOnlineSchema, Error, TEXT("Invalid schema attribute %s: A valid type must be set."),
-				*SchemaAttributeDescriptor.Id.ToString().ToLower());
+				*SchemaAttributeDescriptor.Id.ToString());
 			ParsedSuccessfully = false;
 		}
 
@@ -248,7 +306,7 @@ bool FSchemaRegistry::ParseConfig(const FSchemaRegistryDescriptorConfig& Config)
 		if (SchemaAttributeDescriptor.Type == ESchemaAttributeType::String && SchemaAttributeDescriptor.MaxSize <= 0)
 		{
 			UE_LOG(LogOnlineSchema, Error, TEXT("Invalid schema attribute %s: A valid max size must be set for variable length data."),
-				*SchemaAttributeDescriptor.Id.ToString().ToLower());
+				*SchemaAttributeDescriptor.Id.ToString());
 			ParsedSuccessfully = false;
 		}
 
@@ -256,7 +314,7 @@ bool FSchemaRegistry::ParseConfig(const FSchemaRegistryDescriptorConfig& Config)
 		if (SchemaAttributeDescriptor.UpdateGroupId < 0)
 		{
 			UE_LOG(LogOnlineSchema, Error, TEXT("Invalid schema attribute %s: UpdateGroupId must be set to either 0 or a positive integer value."),
-				*SchemaAttributeDescriptor.Id.ToString().ToLower());
+				*SchemaAttributeDescriptor.Id.ToString());
 			ParsedSuccessfully = false;
 		}
 		
@@ -267,14 +325,14 @@ bool FSchemaRegistry::ParseConfig(const FSchemaRegistryDescriptorConfig& Config)
 			if (SchemaAttributeDescriptor.Type != ESchemaAttributeType::Int64)
 			{
 				UE_LOG(LogOnlineSchema, Error, TEXT("Invalid schema attribute %s: SchemaCompatibilityId field may only be set to Int64."),
-					*SchemaAttributeDescriptor.Id.ToString().ToLower());
+					*SchemaAttributeDescriptor.Id.ToString());
 				ParsedSuccessfully = false;
 			}
 		
 			if (EnumHasAnyFlags(AttributeFlags, ESchemaAttributeFlags::Private))
 			{
 				UE_LOG(LogOnlineSchema, Error, TEXT("Invalid schema attribute %s: SchemaCompatibilityId attribute may not be private."),
-					*SchemaAttributeDescriptor.Id.ToString().ToLower());
+					*SchemaAttributeDescriptor.Id.ToString());
 				ParsedSuccessfully = false;
 			}
 		}
@@ -283,19 +341,19 @@ bool FSchemaRegistry::ParseConfig(const FSchemaRegistryDescriptorConfig& Config)
 		if (EnumHasAllFlags(AttributeFlags, ESchemaAttributeFlags::Public | ESchemaAttributeFlags::Private))
 		{
 			UE_LOG(LogOnlineSchema, Error, TEXT("Invalid schema attribute %s: Public and Private visibility are mutually exclusive."),
-				*SchemaAttributeDescriptor.Id.ToString().ToLower());
+				*SchemaAttributeDescriptor.Id.ToString());
 			ParsedSuccessfully = false;
 		}
 		if (!EnumHasAnyFlags(AttributeFlags, ESchemaAttributeFlags::Public | ESchemaAttributeFlags::Private))
 		{
 			UE_LOG(LogOnlineSchema, Error, TEXT("Invalid schema attribute %s: Either Public or Private visibility must be set."),
-				*SchemaAttributeDescriptor.Id.ToString().ToLower());
+				*SchemaAttributeDescriptor.Id.ToString());
 			ParsedSuccessfully = false;
 		}
 		if (EnumHasAllFlags(AttributeFlags, ESchemaAttributeFlags::Searchable | ESchemaAttributeFlags::Private))
 		{
 			UE_LOG(LogOnlineSchema, Error, TEXT("Invalid schema attribute %s: Searchable attributes may not be private."),
-				*SchemaAttributeDescriptor.Id.ToString().ToLower());
+				*SchemaAttributeDescriptor.Id.ToString());
 			ParsedSuccessfully = false;
 		}
 	}
@@ -311,8 +369,8 @@ bool FSchemaRegistry::ParseConfig(const FSchemaRegistryDescriptorConfig& Config)
 		if (FoundSchema && FoundSchema->Contains(SchemaCategoryAttributesDescriptor.CategoryId))
 		{
 			UE_LOG(LogOnlineSchema, Error, TEXT("Duplicate schema additional category attributes descriptor found: %s.%s"),
-				*SchemaCategoryAttributesDescriptor.SchemaId.ToString().ToLower(),
-				*SchemaCategoryAttributesDescriptor.CategoryId.ToString().ToLower());
+				*SchemaCategoryAttributesDescriptor.SchemaId.ToString(),
+				*SchemaCategoryAttributesDescriptor.CategoryId.ToString());
 			ParsedSuccessfully = false;
 			continue;
 		}
@@ -321,8 +379,8 @@ bool FSchemaRegistry::ParseConfig(const FSchemaRegistryDescriptorConfig& Config)
 		if (!KnownSchemaDescriptors.Contains(SchemaCategoryAttributesDescriptor.SchemaId))
 		{
 			UE_LOG(LogOnlineSchema, Error, TEXT("Invalid additional category attributes %s.%s: Schema does not exist"),
-				*SchemaCategoryAttributesDescriptor.SchemaId.ToString().ToLower(),
-				*SchemaCategoryAttributesDescriptor.CategoryId.ToString().ToLower());
+				*SchemaCategoryAttributesDescriptor.SchemaId.ToString(),
+				*SchemaCategoryAttributesDescriptor.CategoryId.ToString());
 			ParsedSuccessfully = false;
 			continue;
 		}
@@ -331,8 +389,8 @@ bool FSchemaRegistry::ParseConfig(const FSchemaRegistryDescriptorConfig& Config)
 		if (!KnownSchemaCategoryDescriptors.Contains(SchemaCategoryAttributesDescriptor.CategoryId))
 		{
 			UE_LOG(LogOnlineSchema, Error, TEXT("Invalid additional category attributes %s.%s: Category does not exist"),
-				*SchemaCategoryAttributesDescriptor.SchemaId.ToString().ToLower(),
-				*SchemaCategoryAttributesDescriptor.CategoryId.ToString().ToLower());
+				*SchemaCategoryAttributesDescriptor.SchemaId.ToString(),
+				*SchemaCategoryAttributesDescriptor.CategoryId.ToString());
 			ParsedSuccessfully = false;
 			continue;
 		}
@@ -344,9 +402,9 @@ bool FSchemaRegistry::ParseConfig(const FSchemaRegistryDescriptorConfig& Config)
 			if (SeenAttributeIds.Contains(AttributeId))
 			{
 				UE_LOG(LogOnlineSchema, Error, TEXT("Invalid category attributes %s.%s: Attribute %s duplicated"),
-					*SchemaCategoryAttributesDescriptor.SchemaId.ToString().ToLower(),
-					*SchemaCategoryAttributesDescriptor.CategoryId.ToString().ToLower(),
-					*AttributeId.ToString().ToLower());
+					*SchemaCategoryAttributesDescriptor.SchemaId.ToString(),
+					*SchemaCategoryAttributesDescriptor.CategoryId.ToString(),
+					*AttributeId.ToString());
 				bAttributesParsedSuccessfully = false;
 				continue;
 			}
@@ -355,9 +413,9 @@ bool FSchemaRegistry::ParseConfig(const FSchemaRegistryDescriptorConfig& Config)
 			if (!KnownSchemaAttributeDescriptors.Contains(AttributeId))
 			{
 				UE_LOG(LogOnlineSchema, Error, TEXT("Invalid category attributes %s.%s: Attribute %s does not exist"),
-					*SchemaCategoryAttributesDescriptor.SchemaId.ToString().ToLower(),
-					*SchemaCategoryAttributesDescriptor.CategoryId.ToString().ToLower(),
-					*AttributeId.ToString().ToLower());
+					*SchemaCategoryAttributesDescriptor.SchemaId.ToString(),
+					*SchemaCategoryAttributesDescriptor.CategoryId.ToString(),
+					*AttributeId.ToString());
 				bAttributesParsedSuccessfully = false;
 				continue;
 			}
@@ -390,7 +448,7 @@ bool FSchemaRegistry::ParseConfig(const FSchemaRegistryDescriptorConfig& Config)
 			// Check for circular dependencies in schema.
 			if (SchemaHierarchy.Contains(TestSchemaDescriptor->Id))
 			{
-				UE_LOG(LogOnlineSchema, Error, TEXT("Circular parent dependency found in schema: %s"), *SchemaDescriptorToVerify.Id.ToString().ToLower());
+				UE_LOG(LogOnlineSchema, Error, TEXT("Circular parent dependency found in schema: %s"), *SchemaDescriptorToVerify.Id.ToString());
 				ParsedSuccessfully = false;
 				break;
 			}
@@ -406,8 +464,8 @@ bool FSchemaRegistry::ParseConfig(const FSchemaRegistryDescriptorConfig& Config)
 				else
 				{
 					UE_LOG(LogOnlineSchema, Error, TEXT("Invalid schema %s: Parent schema %s does not exist."),
-						*SchemaDescriptorToVerify.Id.ToString().ToLower(),
-						*TestSchemaDescriptor->ParentId.ToString().ToLower());
+						*SchemaDescriptorToVerify.Id.ToString(),
+						*TestSchemaDescriptor->ParentId.ToString());
 					ParsedSuccessfully = false;
 					break;
 				}
@@ -443,9 +501,9 @@ bool FSchemaRegistry::ParseConfig(const FSchemaRegistryDescriptorConfig& Config)
 				if (SeenSchemaServiceDescriptorIds.Find(SchemaCategoryDescriptor.ServiceDescriptorId) != nullptr)
 				{
 					UE_LOG(LogOnlineSchema, Error, TEXT("Invalid schema category %s.%s: Service descriptor id %s has already been used by another category."),
-						*SchemaDescriptorToVerify.Id.ToString().ToLower(),
-						*SchemaCategoryDescriptor.Id.ToString().ToLower(),
-						*SchemaCategoryDescriptor.ServiceDescriptorId.ToString().ToLower());
+						*SchemaDescriptorToVerify.Id.ToString(),
+						*SchemaCategoryDescriptor.Id.ToString(),
+						*SchemaCategoryDescriptor.ServiceDescriptorId.ToString());
 					ParsedSuccessfully = false;
 				}
 
@@ -473,9 +531,9 @@ bool FSchemaRegistry::ParseConfig(const FSchemaRegistryDescriptorConfig& Config)
 					if (SeenAttributes.Find(SchemaAttributeId))
 					{
 						UE_LOG(LogOnlineSchema, Error, TEXT("Invalid schema attribute %s.%s.%s: Attribute id has already been used."),
-							*SchemaDescriptorToVerify.Id.ToString().ToLower(),
-							*SchemaCategoryDescriptor.Id.ToString().ToLower(),
-							*SchemaAttributeId.ToString().ToLower());
+							*SchemaDescriptorToVerify.Id.ToString(),
+							*SchemaCategoryDescriptor.Id.ToString(),
+							*SchemaAttributeId.ToString());
 						ParsedSuccessfully = false;
 						continue;
 					}
@@ -494,9 +552,9 @@ bool FSchemaRegistry::ParseConfig(const FSchemaRegistryDescriptorConfig& Config)
 					if (EnumHasAllFlags(AttributeFlags, ESchemaAttributeFlags::Searchable) && SchemaHierarchyLayerId != RootSchemaId)
 					{
 						UE_LOG(LogOnlineSchema, Error, TEXT("Invalid schema attribute %s.%s.%s: Searchable fields may only exist in the base schema."),
-							*SchemaDescriptorToVerify.Id.ToString().ToLower(),
-							*SchemaCategoryDescriptor.Id.ToString().ToLower(),
-							*SchemaAttributeDescriptor.Id.ToString().ToLower());
+							*SchemaDescriptorToVerify.Id.ToString(),
+							*SchemaCategoryDescriptor.Id.ToString(),
+							*SchemaAttributeDescriptor.Id.ToString());
 						ParsedSuccessfully = false;
 					}
 
@@ -506,19 +564,19 @@ bool FSchemaRegistry::ParseConfig(const FSchemaRegistryDescriptorConfig& Config)
 						if (SchemaHierarchyLayerId != RootSchemaId)
 						{
 							UE_LOG(LogOnlineSchema, Error, TEXT("Invalid schema attribute %s.%s.%s: SchemaCompatibilityId field may only exist in the base schema."),
-								*SchemaDescriptorToVerify.Id.ToString().ToLower(),
-								*SchemaCategoryDescriptor.Id.ToString().ToLower(),
-								*SchemaAttributeDescriptor.Id.ToString().ToLower());
+								*SchemaDescriptorToVerify.Id.ToString(),
+								*SchemaCategoryDescriptor.Id.ToString(),
+								*SchemaAttributeDescriptor.Id.ToString());
 							ParsedSuccessfully = false;
 						}
 
 						if (SeenSchemaCompatibilityAttributeId != FSchemaId())
 						{
 							UE_LOG(LogOnlineSchema, Error, TEXT("Invalid schema attribute %s.%s.%s: SchemaCompatibilityId field has already been used by attribute %s."),
-								*SchemaDescriptorToVerify.Id.ToString().ToLower(),
-								*SchemaCategoryDescriptor.Id.ToString().ToLower(),
-								*SchemaAttributeDescriptor.Id.ToString().ToLower(),
-								*SeenSchemaCompatibilityAttributeId.ToString().ToLower());
+								*SchemaDescriptorToVerify.Id.ToString(),
+								*SchemaCategoryDescriptor.Id.ToString(),
+								*SchemaAttributeDescriptor.Id.ToString(),
+								*SeenSchemaCompatibilityAttributeId.ToString());
 							ParsedSuccessfully = false;
 						}
 						else
@@ -724,9 +782,9 @@ bool FSchemaRegistry::ParseConfig(const FSchemaRegistryDescriptorConfig& Config)
 							else
 							{
 								UE_LOG(LogOnlineSchema, Error, TEXT("Failed to find service attribute to fit schema attribute %s.%s.%s."),
-									*SchemaDefinition->Id.ToString().ToLower(),
-									*SchemaCategoryDefinition.Id.ToString().ToLower(),
-									*SchemaAttributeDefinition.Id.ToString().ToLower());
+									*SchemaDefinition->Id.ToString(),
+									*SchemaCategoryDefinition.Id.ToString(),
+									*SchemaAttributeDefinition.Id.ToString());
 								ParsedSuccessfully = false;
 							}
 						}
@@ -743,7 +801,7 @@ bool FSchemaRegistry::ParseConfig(const FSchemaRegistryDescriptorConfig& Config)
 				if (SeenSchemaCrcs.Find(SchemaIdCrc) != nullptr)
 				{
 					UE_LOG(LogOnlineSchema, Error, TEXT("Invalid schema %s: CRC collision processing schema name, please rename the schema to avoid collision."),
-						*SchemaDefinition->Id.ToString().ToLower());
+						*SchemaDefinition->Id.ToString());
 					ParsedSuccessfully = false;
 				}
 				SeenSchemaCrcs.Add(SchemaIdCrc);
@@ -751,14 +809,16 @@ bool FSchemaRegistry::ParseConfig(const FSchemaRegistryDescriptorConfig& Config)
 				const uint64 CompatibilityIdUint = (static_cast<uint64>(SchemaIdCrc) << 32) | SchemaDataCrc;
 				SchemaDefinition->CompatibilityId = *reinterpret_cast<const int64*>(&CompatibilityIdUint);
 
-				UE_LOG(LogOnlineSchema, Verbose, TEXT("Building compatability id: SchemaId: %s, SchemaIdCrc: 0x%04X, SchemaDataCrc 0x%04X."),
-					*SchemaDefinition->Id.ToString().ToLower(), SchemaIdCrc, SchemaDataCrc);
+				UE_LOG(LogOnlineSchema, Verbose, TEXT("Building compatability id: SchemaId: [%s], SchemaIdCrc: [0x%04X], SchemaDataCrc [0x%04X]."),
+					*SchemaDefinition->Id.ToString(), SchemaIdCrc, SchemaDataCrc);
 			}
 
 			ParsedSchemaDefinitionsByCompatibilityId.Add(SchemaDefinition->CompatibilityId, SchemaDefinition);
 
-			UE_LOG(LogOnlineSchema, Log, TEXT("Parsed Schema definition: SchemaId: %s, SchemaCompatabilityId: 0x%08" INT64_X_FMT "."),
-				*SchemaDefinition->Id.ToString().ToLower(), SchemaDefinition->CompatibilityId);
+			UE_LOG(LogOnlineSchema, Log, TEXT("Parsed Schema definition: SchemaId: [%s], SchemaCompatabilityId: [0x%08" INT64_X_FMT "]"),
+				*SchemaDefinition->Id.ToString(), SchemaDefinition->CompatibilityId);
+
+			LogSchemaDefinition(*SchemaDefinition);
 		}
 	}
 
@@ -807,9 +867,9 @@ FSchemaCategoryInstanceBase::FSchemaCategoryInstanceBase(
 	if (!bInitSuccess)
 	{
 		UE_LOG(LogOnlineSchema, Error, TEXT("[FSchemaCategoryInstance] init failed. DerivedSchema[%s], BaseSchema[%s], CategoryId[%s]"),
-			*DerivedSchemaId.ToString().ToLower(),
-			*BaseSchemaId.ToString().ToLower(),
-			*CategoryId.ToString().ToLower());
+			*DerivedSchemaId.ToString(),
+			*BaseSchemaId.ToString(),
+			*CategoryId.ToString());
 
 		// Invalidate schema to fail initialization.
 		BaseSchemaDefinition.Reset();
@@ -864,8 +924,8 @@ TOnlineResult<FSchemaCategoryInstancePrepareClientChanges> FSchemaCategoryInstan
 		else
 		{
 			UE_LOG(LogOnlineSchema, Warning, TEXT("[PrepareClientChanges] Schema change failed: NewDerivedSchema[%s], BaseSchema[%s]"),
-				*Params.ClientChanges.SchemaId->ToString().ToLower(),
-				*BaseSchemaCategoryDefinition->Id.ToString().ToLower());
+				*Params.ClientChanges.SchemaId->ToString(),
+				*BaseSchemaCategoryDefinition->Id.ToString());
 			return TOnlineResult<FSchemaCategoryInstancePrepareClientChanges>(Errors::InvalidParams());
 		}
 	}
@@ -880,9 +940,9 @@ TOnlineResult<FSchemaCategoryInstancePrepareClientChanges> FSchemaCategoryInstan
 		if (AttributeDefinition == nullptr)
 		{
 			UE_LOG(LogOnlineSchema, Warning, TEXT("[PrepareClientChanges] %s.%s.%s: Attribute definition does not exist in schema category."),
-				*SchemaDefinition->Id.ToString().ToLower(),
-				*SchemaCategoryDefinition->Id.ToString().ToLower(),
-				*UpdatedAttributeId.ToString().ToLower());
+				*SchemaDefinition->Id.ToString(),
+				*SchemaCategoryDefinition->Id.ToString(),
+				*UpdatedAttributeId.ToString());
 			return TOnlineResult<FSchemaCategoryInstancePrepareClientChanges>(Errors::InvalidParams());
 		}
 
@@ -890,9 +950,9 @@ TOnlineResult<FSchemaCategoryInstancePrepareClientChanges> FSchemaCategoryInstan
 		if (EnumHasAnyFlags(AttributeDefinition->Flags, ESchemaAttributeFlags::SchemaCompatibilityId))
 		{
 			UE_LOG(LogOnlineSchema, Warning, TEXT("[PrepareClientChanges] %s.%s.%s: It is not valid to set SchemaCompatibilityId in an attribute update"),
-				*SchemaDefinition->Id.ToString().ToLower(),
-				*SchemaCategoryDefinition->Id.ToString().ToLower(),
-				*AttributeDefinition->Id.ToString().ToLower());
+				*SchemaDefinition->Id.ToString(),
+				*SchemaCategoryDefinition->Id.ToString(),
+				*AttributeDefinition->Id.ToString());
 			return TOnlineResult<FSchemaCategoryInstancePrepareClientChanges>(Errors::InvalidParams());
 		}
 
@@ -900,9 +960,9 @@ TOnlineResult<FSchemaCategoryInstancePrepareClientChanges> FSchemaCategoryInstan
 		if (AttributeDefinition->Type != UpdatedAttributeValue.GetType())
 		{
 			UE_LOG(LogOnlineSchema, Warning, TEXT("[PrepareClientChanges] %s.%s.%s: Attribute type %s does not match the expected type %s."),
-				*SchemaDefinition->Id.ToString().ToLower(),
-				*SchemaCategoryDefinition->Id.ToString().ToLower(),
-				*AttributeDefinition->Id.ToString().ToLower(),
+				*SchemaDefinition->Id.ToString(),
+				*SchemaCategoryDefinition->Id.ToString(),
+				*AttributeDefinition->Id.ToString(),
 				LexToString(UpdatedAttributeValue.GetType()),
 				LexToString(AttributeDefinition->Type));
 			return TOnlineResult<FSchemaCategoryInstancePrepareClientChanges>(Errors::InvalidParams());
@@ -912,9 +972,9 @@ TOnlineResult<FSchemaCategoryInstancePrepareClientChanges> FSchemaCategoryInstan
 		if (AttributeDefinition->Type == ESchemaAttributeType::String && UpdatedAttributeValue.GetString().Len() > AttributeDefinition->MaxSize)
 		{
 			UE_LOG(LogOnlineSchema, Warning, TEXT("[PrepareClientChanges] %s.%s.%s: Variably sized attribute exceeds maximum defined length. %d > %d."),
-				*SchemaDefinition->Id.ToString().ToLower(),
-				*SchemaCategoryDefinition->Id.ToString().ToLower(),
-				*AttributeDefinition->Id.ToString().ToLower(),
+				*SchemaDefinition->Id.ToString(),
+				*SchemaCategoryDefinition->Id.ToString(),
+				*AttributeDefinition->Id.ToString(),
 				UpdatedAttributeValue.GetString().Len(),
 				AttributeDefinition->MaxSize);
 			return TOnlineResult<FSchemaCategoryInstancePrepareClientChanges>(Errors::InvalidParams());
@@ -947,10 +1007,10 @@ TOnlineResult<FSchemaCategoryInstancePrepareClientChanges> FSchemaCategoryInstan
 			else
 			{
 				UE_LOG(LogOnlineSchema, Warning, TEXT("[PrepareClientChanges] %s.%s.%s: Unable to find service attribute definition for service attribute %s."),
-					*SchemaDefinition->Id.ToString().ToLower(),
-					*SchemaCategoryDefinition->Id.ToString().ToLower(),
-					*AttributeDefinition->Id.ToString().ToLower(),
-					*AttributeDefinition->ServiceAttributeId.ToString().ToLower());
+					*SchemaDefinition->Id.ToString(),
+					*SchemaCategoryDefinition->Id.ToString(),
+					*AttributeDefinition->Id.ToString(),
+					*AttributeDefinition->ServiceAttributeId.ToString());
 				return TOnlineResult<FSchemaCategoryInstancePrepareClientChanges>(Errors::InvalidState());
 			}
 		}
@@ -965,9 +1025,9 @@ TOnlineResult<FSchemaCategoryInstancePrepareClientChanges> FSchemaCategoryInstan
 			if (EnumHasAnyFlags(AttributeDefinition->Flags, ESchemaAttributeFlags::SchemaCompatibilityId))
 			{
 				UE_LOG(LogOnlineSchema, Warning, TEXT("[PrepareClientChanges] %s.%s.%s: It is not valid to set SchemaCompatibilityId in an attribute update"),
-					*SchemaDefinition->Id.ToString().ToLower(),
-					*SchemaCategoryDefinition->Id.ToString().ToLower(),
-					*AttributeDefinition->Id.ToString().ToLower());
+					*SchemaDefinition->Id.ToString(),
+					*SchemaCategoryDefinition->Id.ToString(),
+					*AttributeDefinition->Id.ToString());
 				return TOnlineResult<FSchemaCategoryInstancePrepareClientChanges>(Errors::InvalidParams());
 			}
 			else
@@ -985,9 +1045,9 @@ TOnlineResult<FSchemaCategoryInstancePrepareClientChanges> FSchemaCategoryInstan
 		else
 		{
 			UE_LOG(LogOnlineSchema, Warning, TEXT("[PrepareClientChanges] %s.%s.%s: Attribute definition does not exist in schema category."),
-				*SchemaDefinition->Id.ToString().ToLower(),
-				*SchemaCategoryDefinition->Id.ToString().ToLower(),
-				*RemovedAttributeId.ToString().ToLower());
+				*SchemaDefinition->Id.ToString(),
+				*SchemaCategoryDefinition->Id.ToString(),
+				*RemovedAttributeId.ToString());
 			return TOnlineResult<FSchemaCategoryInstancePrepareClientChanges>(Errors::InvalidParams());
 		}
 	}
@@ -1008,10 +1068,10 @@ TOnlineResult<FSchemaCategoryInstancePrepareClientChanges> FSchemaCategoryInstan
 				else
 				{
 					UE_LOG(LogOnlineSchema, Warning, TEXT("[PrepareClientChanges] %s.%s.%s: Unable to find service attribute definition for service attribute %s."),
-						*SchemaDefinition->Id.ToString().ToLower(),
-						*SchemaCategoryDefinition->Id.ToString().ToLower(),
-						*AttributeDefinition->Id.ToString().ToLower(),
-						*AttributeDefinition->ServiceAttributeId.ToString().ToLower());
+						*SchemaDefinition->Id.ToString(),
+						*SchemaCategoryDefinition->Id.ToString(),
+						*AttributeDefinition->Id.ToString(),
+						*AttributeDefinition->ServiceAttributeId.ToString());
 					return TOnlineResult<FSchemaCategoryInstancePrepareClientChanges>(Errors::InvalidState());
 				}
 			}
@@ -1021,9 +1081,9 @@ TOnlineResult<FSchemaCategoryInstancePrepareClientChanges> FSchemaCategoryInstan
 		else
 		{
 			UE_LOG(LogOnlineSchema, Warning, TEXT("[PrepareClientChanges] %s.%s.%s: Failed to find attribute definition for SchemaCompatibilityId field."),
-				*SchemaDefinition->Id.ToString().ToLower(),
-				*SchemaCategoryDefinition->Id.ToString().ToLower(),
-				*SchemaCategoryDefinition->SchemaCompatibilityAttributeId.ToString().ToLower());
+				*SchemaDefinition->Id.ToString(),
+				*SchemaCategoryDefinition->Id.ToString(),
+				*SchemaCategoryDefinition->SchemaCompatibilityAttributeId.ToString());
 			return TOnlineResult<FSchemaCategoryInstancePrepareClientChanges>(Errors::InvalidState());
 		}
 	}
@@ -1141,9 +1201,9 @@ TOnlineResult<FSchemaCategoryInstancePrepareServiceSnapshot> FSchemaCategoryInst
 		else
 		{
 			UE_LOG(LogOnlineSchema, Warning, TEXT("[PrepareServiceSnapshot] %s.%s.%s: Failed to find valid definition for SchemaCompatibilityId 0x%08" INT64_X_FMT "."),
-				*SchemaDefinition->Id.ToString().ToLower(),
-				*SchemaCategoryDefinition->Id.ToString().ToLower(),
-				*SchemaCategoryDefinition->SchemaCompatibilityServiceAttributeId.ToString().ToLower(),
+				*SchemaDefinition->Id.ToString(),
+				*SchemaCategoryDefinition->Id.ToString(),
+				*SchemaCategoryDefinition->SchemaCompatibilityServiceAttributeId.ToString(),
 				SchemaCompatibilityId);
 			return TOnlineResult<FSchemaCategoryInstancePrepareServiceSnapshot>(Errors::InvalidParams());
 		}
@@ -1156,8 +1216,8 @@ TOnlineResult<FSchemaCategoryInstancePrepareServiceSnapshot> FSchemaCategoryInst
 		if (SchemaCompatibilityAttributeData)
 		{
 			UE_LOG(LogOnlineSchema, Warning, TEXT("[PrepareServiceSnapshot] %s.%s: Unable to handle user schema swap when service defines a schema compatibility id."),
-				*SchemaDefinition->Id.ToString().ToLower(),
-				*SchemaCategoryDefinition->Id.ToString().ToLower());
+				*SchemaDefinition->Id.ToString(),
+				*SchemaCategoryDefinition->Id.ToString());
 			return TOnlineResult<FSchemaCategoryInstancePrepareServiceSnapshot>(Errors::InvalidParams());
 		}
 	
@@ -1182,9 +1242,9 @@ TOnlineResult<FSchemaCategoryInstancePrepareServiceSnapshot> FSchemaCategoryInst
 		else
 		{
 			UE_LOG(LogOnlineSchema, Warning, TEXT("[PrepareServiceSnapshot] %s.%s: Failed to find valid definition for SchemaId %s."),
-				*SchemaDefinition->Id.ToString().ToLower(),
-				*SchemaCategoryDefinition->Id.ToString().ToLower(),
-				*Params.SchemaId->ToString().ToLower());
+				*SchemaDefinition->Id.ToString(),
+				*SchemaCategoryDefinition->Id.ToString(),
+				*Params.SchemaId->ToString());
 			return TOnlineResult<FSchemaCategoryInstancePrepareServiceSnapshot>(Errors::InvalidParams());
 		}
 	}
@@ -1212,9 +1272,9 @@ TOnlineResult<FSchemaCategoryInstancePrepareServiceSnapshot> FSchemaCategoryInst
 				if (EnumHasAnyFlags(SchemaAttributeDefinition->Flags, ESchemaAttributeFlags::SchemaCompatibilityId))
 				{
 					UE_LOG(LogOnlineSchema, Verbose, TEXT("[PrepareServiceSnapshot] %s.%s.%s: Consuming SchemaCompatibility attribute: 0x%08" INT64_X_FMT "."),
-						*SchemaDefinition->Id.ToString().ToLower(),
-						*SchemaCategoryDefinition->Id.ToString().ToLower(),
-						*SchemaAttributeDefinition->Id.ToString().ToLower(),
+						*SchemaDefinition->Id.ToString(),
+						*SchemaCategoryDefinition->Id.ToString(),
+						*SchemaAttributeDefinition->Id.ToString(),
 						SchemaServiceAttributeValue.GetInt64());
 				}
 				else
@@ -1226,10 +1286,11 @@ TOnlineResult<FSchemaCategoryInstancePrepareServiceSnapshot> FSchemaCategoryInst
 		else
 		{
 			UE_LOG(LogOnlineSchema, Warning, TEXT("[PrepareServiceSnapshot] %s.%s.%s: Service attribute definition does not exist in schema category."),
-				*SchemaDefinition->Id.ToString().ToLower(),
-				*SchemaCategoryDefinition->Id.ToString().ToLower(),
-				*SchemaServiceAttributeId.ToString().ToLower());
-			return TOnlineResult<FSchemaCategoryInstancePrepareServiceSnapshot>(Errors::InvalidParams());
+				*SchemaDefinition->Id.ToString(),
+				*SchemaCategoryDefinition->Id.ToString(),
+				*SchemaServiceAttributeId.ToString());
+			// Allow code to continue when an unexpected attribute is received.
+			continue;
 		}
 	}
 
@@ -1335,10 +1396,10 @@ bool FSchemaCategoryInstanceBase::VerifyBaseAttributeData(
 						else
 						{
 							UE_LOG(LogOnlineSchema, Warning, TEXT("[VerifyBaseAttributeData] %s.%s.%s: Unable to find service attribute definition for service attribute %s."),
-								*BaseSchemaCategoryDefinition->Id.ToString().ToLower(),
-								*BaseSchemaCategoryDefinition->Id.ToString().ToLower(),
-								*AttributeDefinition->Id.ToString().ToLower(),
-								*AttributeDefinition->ServiceAttributeId.ToString().ToLower());
+								*BaseSchemaCategoryDefinition->Id.ToString(),
+								*BaseSchemaCategoryDefinition->Id.ToString(),
+								*AttributeDefinition->Id.ToString(),
+								*AttributeDefinition->ServiceAttributeId.ToString());
 							return false;
 						}
 					}
@@ -1353,9 +1414,9 @@ bool FSchemaCategoryInstanceBase::VerifyBaseAttributeData(
 			else
 			{
 				UE_LOG(LogOnlineSchema, Verbose, TEXT("[VerifyBaseAttributeData] Schema attribute %s.%s.%s set with invalid type %s. Expected type %s"),
-					*BaseSchemaDefinition->Id.ToString().ToLower(),
-					*BaseSchemaCategoryDefinition->Id.ToString().ToLower(),
-					*AttributeDefinition->Id.ToString().ToLower(),
+					*BaseSchemaDefinition->Id.ToString(),
+					*BaseSchemaCategoryDefinition->Id.ToString(),
+					*AttributeDefinition->Id.ToString(),
 					LexToString(Data.GetType()),
 					LexToString(AttributeDefinition->Type));
 				return false;
@@ -1364,16 +1425,16 @@ bool FSchemaCategoryInstanceBase::VerifyBaseAttributeData(
 		else
 		{
 			UE_LOG(LogOnlineSchema, Verbose, TEXT("[VerifyBaseAttributeData] Attribute %s not found in schema definition %s.%s"),
-				*Id.ToString().ToLower(),
-				*BaseSchemaDefinition->Id.ToString().ToLower(),
-				*BaseSchemaCategoryDefinition->Id.ToString().ToLower());
+				*Id.ToString(),
+				*BaseSchemaDefinition->Id.ToString(),
+				*BaseSchemaCategoryDefinition->Id.ToString());
 			return false;
 		}
 	}
 	else
 	{
 		UE_LOG(LogOnlineSchema, Verbose, TEXT("[VerifyBaseAttributeData] Unable to set attribute %s. Schema is not valid"),
-			*Id.ToString().ToLower());
+			*Id.ToString());
 		return false;
 	}
 }
@@ -1416,7 +1477,7 @@ bool FSchemaCategoryInstanceBase::InitializeSchemaDefinition(
 	else
 	{
 		UE_LOG(LogOnlineSchema, Error, TEXT("[InitializeSchemaDefinition] init error: Unable to find definition for schema %s"),
-			*SchemaId.ToString().ToLower());
+			*SchemaId.ToString());
 		return false;
 	}
 }
@@ -1435,8 +1496,8 @@ bool FSchemaCategoryInstanceBase::InitializeSchemaDefinition(
 	if (OptionalBaseDefinition && !SchemaRegistry->IsSchemaChildOf(NewDefinition->Id, OptionalBaseDefinition->Id))
 	{
 		UE_LOG(LogOnlineSchema, Warning, TEXT("[InitializeSchemaDefinition] Invalid schema definition. New schema %s is not a child of %s."),
-			*NewDefinition->Id.ToString().ToLower(),
-			*OptionalBaseDefinition->Id.ToString().ToLower());
+			*NewDefinition->Id.ToString(),
+			*OptionalBaseDefinition->Id.ToString());
 		return false;
 	}
 
@@ -1445,8 +1506,8 @@ bool FSchemaCategoryInstanceBase::InitializeSchemaDefinition(
 	if (SchemaCategoryDefinition == nullptr)
 	{
 		UE_LOG(LogOnlineSchema, Error, TEXT("[InitializeSchemaDefinition] init error: Unable to find category %s in schema %s"),
-			*CategoryId.ToString().ToLower(),
-			*NewDefinition->Id.ToString().ToLower());
+			*CategoryId.ToString(),
+			*NewDefinition->Id.ToString());
 		return false;
 	}
 
@@ -1513,4 +1574,4 @@ bool FSchemaRegistry::IsSchemaChildOf(const FSchemaId& SchemaId, const FSchemaId
 	return FoundSchemaId != nullptr;
 }
 
-/* UE::Online */ }
+/* UE::Online */ }
\ No newline at end of file
diff --git a/Engine/Plugins/Online/OnlineServices/Source/OnlineServicesInterface/Private/Online/SchemaTypes.cpp b/Engine/Plugins/Online/OnlineServices/Source/OnlineServicesInterface/Private/Online/SchemaTypes.cpp
--- a/Engine/Plugins/Online/OnlineServices/Source/OnlineServicesInterface/Private/Online/SchemaTypes.cpp
+++ b/Engine/Plugins/Online/OnlineServices/Source/OnlineServicesInterface/Private/Online/SchemaTypes.cpp
@@ -94,11 +94,12 @@ const TCHAR* LexToString(ESchemaAttributeFlags SchemaAttributeFlags)
 {
 	switch (SchemaAttributeFlags)
 	{
-	default:								checkNoEntry(); // Intentional fallthrough
-	case ESchemaAttributeFlags::None:		return TEXT("None");
-	case ESchemaAttributeFlags::Searchable:	return TEXT("Searchable");
-	case ESchemaAttributeFlags::Public:		return TEXT("Public");
-	case ESchemaAttributeFlags::Private:	return TEXT("Private");
+	default:											checkNoEntry(); // Intentional fallthrough
+	case ESchemaAttributeFlags::None:					return TEXT("None");
+	case ESchemaAttributeFlags::Searchable:				return TEXT("Searchable");
+	case ESchemaAttributeFlags::Public:					return TEXT("Public");
+	case ESchemaAttributeFlags::Private:				return TEXT("Private");
+	case ESchemaAttributeFlags::SchemaCompatibilityId:	return TEXT("SchemaCompatibilityId");
 	}
 }
 
diff --git a/Engine/Source/Programs/Online/OnlineTestsCore/Source/Private/Tests/Schema/SchemaTests.cpp b/Engine/Source/Programs/Online/OnlineTestsCore/Source/Private/Tests/Schema/SchemaTests.cpp
--- a/Engine/Source/Programs/Online/OnlineTestsCore/Source/Private/Tests/Schema/SchemaTests.cpp
+++ b/Engine/Source/Programs/Online/OnlineTestsCore/Source/Private/Tests/Schema/SchemaTests.cpp
@@ -1676,7 +1676,7 @@ SCHEMA_TEST_CASE("Schema category instance")
 		CHECK(PrepareServiceSnapshotResult.GetErrorValue() == Errors::InvalidParams());
 	}
 
-	SECTION("Valid instance - backing service - Read unknown service attribute")
+	SECTION("Valid instance - backing service - Read unknown schema")
 	{
 		TSharedRef<FSchemaRegistry> SchemaRegistry1 = MakeShared<FSchemaRegistry>();
 		{
@@ -1747,4 +1747,69 @@ SCHEMA_TEST_CASE("Schema category instance")
 		REQUIRE(!PrepareServiceSnapshotResult.IsOk());
 		CHECK(PrepareServiceSnapshotResult.GetErrorValue() == Errors::InvalidParams());
 	}
+
+	SECTION("Valid instance - backing service - Read unknown service attribute")
+	{
+		FSchemaRegistryDescriptorConfig SchemaConfig;
+		SchemaConfig.SchemaDescriptors.Add({ BaseSchemaId1, FSchemaId(), { CategoryId1 } });
+		SchemaConfig.SchemaCategoryDescriptors.Add({ CategoryId1, ServiceDescriptorId });
+		SchemaConfig.SchemaCategoryAttributeDescriptors.Add({ BaseSchemaId1, CategoryId1, { AttributeId1 } });
+		SchemaConfig.SchemaAttributeDescriptors.Add({ AttributeId1, ESchemaAttributeType::Int64, { ESchemaAttributeFlags::Public, ESchemaAttributeFlags::SchemaCompatibilityId }, 0, 0 });
+
+		SchemaConfig.SchemaDescriptors.Add({ SchemaId1, BaseSchemaId1 });
+		SchemaConfig.SchemaCategoryAttributeDescriptors.Add({ SchemaId1, CategoryId1, { AttributeId2 } });
+		SchemaConfig.SchemaAttributeDescriptors.Add({ AttributeId2, ESchemaAttributeType::Int64, { ESchemaAttributeFlags::Private }, 0, 0 });
+
+		SchemaConfig.ServiceDescriptors.Add({ ServiceDescriptorId, { ServiceAttributeId1, ServiceAttributeId2 } });
+		SchemaConfig.ServiceAttributeDescriptors.Add({ ServiceAttributeId1, AllSupportedServiceAttributeTypes, AllSupportedServiceAttributeFlags, 32 });
+		SchemaConfig.ServiceAttributeDescriptors.Add({ ServiceAttributeId2, AllSupportedServiceAttributeTypes, AllSupportedServiceAttributeFlags, 32 });
+
+		TSharedRef<FSchemaRegistry> SchemaRegistry = MakeShared<FSchemaRegistry>();
+		REQUIRE(SchemaRegistry->ParseConfig(SchemaConfig));
+
+		FSchemaCategoryInstance CategoryInstance1(SchemaId1, BaseSchemaId1, CategoryId1, SchemaRegistry);
+		REQUIRE(CategoryInstance1.IsValid());
+
+		FSchemaCategoryInstance CategoryInstance2(FSchemaId(), BaseSchemaId1, CategoryId1, SchemaRegistry);
+		REQUIRE(CategoryInstance2.IsValid());
+
+		// Write value to instance 1.
+		FSchemaVariant VariantData(42ll);
+		FSchemaCategoryInstancePrepareClientChanges::Params ClientChangesParams;
+		ClientChangesParams.ClientChanges.UpdatedAttributes = TMap<FSchemaAttributeId, FSchemaVariant>({ { AttributeId2, VariantData } });
+		TOnlineResult<FSchemaCategoryInstancePrepareClientChanges> PrepareClientChangesResult = CategoryInstance1.PrepareClientChanges(MoveTemp(ClientChangesParams));
+		CHECK(PrepareClientChangesResult.IsOk());
+
+		// Grab service output.
+		FSchemaServiceAttributeData* SchemaServiceAttribute1Data = PrepareClientChangesResult.GetOkValue().ServiceChanges.UpdatedAttributes.Find(ServiceAttributeId1);
+		REQUIRE(SchemaServiceAttribute1Data);
+		FSchemaServiceAttributeData* SchemaServiceAttribute2Data = PrepareClientChangesResult.GetOkValue().ServiceChanges.UpdatedAttributes.Find(ServiceAttributeId2);
+		REQUIRE(SchemaServiceAttribute2Data);
+
+		FSchemaServiceAttributeData UnexpectedSchemaServiceAttributeData;
+		UnexpectedSchemaServiceAttributeData.Id = "UnexpectedAttributeId";
+		UnexpectedSchemaServiceAttributeData.Value = TEXT("Yams");
+
+		// Feed service output into instance 2, but include a third unexpected attribute.
+		FSchemaCategoryInstancePrepareServiceSnapshot::Params PrepareServiceSnapshotParams;
+		PrepareServiceSnapshotParams.ServiceSnapshot.Attributes = TMap<FSchemaServiceAttributeId, FSchemaVariant>({
+			{ SchemaServiceAttribute1Data->Id, SchemaServiceAttribute1Data->Value },
+			{ SchemaServiceAttribute2Data->Id, SchemaServiceAttribute2Data->Value },
+			{ UnexpectedSchemaServiceAttributeData.Id, UnexpectedSchemaServiceAttributeData.Value }});
+		TOnlineResult<FSchemaCategoryInstancePrepareServiceSnapshot> PrepareServiceSnapshotResult = CategoryInstance2.PrepareServiceSnapshot(MoveTemp(PrepareServiceSnapshotParams));
+		REQUIRE(PrepareServiceSnapshotResult.IsOk());
+
+		// Check that the schema in instance 2 changed from BaseSchemaId1 to SchemaId1.
+		FSchemaCategoryInstanceCommitServiceSnapshot::Result CommitServiceSnapshotResult = CategoryInstance2.CommitServiceSnapshot();
+		REQUIRE(CommitServiceSnapshotResult.ClientChanges.SchemaId.IsSet());
+		CHECK(*CommitServiceSnapshotResult.ClientChanges.SchemaId == SchemaId1);
+		REQUIRE(CategoryInstance2.GetDerivedDefinition() != nullptr);
+		CHECK(CategoryInstance2.GetDerivedDefinition()->Id == SchemaId1);
+
+		// Check that attribute data defined in SchemaId1 is received in instance 2.
+		CHECK(CommitServiceSnapshotResult.ClientChanges.AddedAttributes.Num() == 1);
+		FSchemaVariant* SchemaAttributeData = CommitServiceSnapshotResult.ClientChanges.AddedAttributes.Find(AttributeId2);
+		REQUIRE(SchemaAttributeData != nullptr);
+		CHECK(*SchemaAttributeData == VariantData);
+	}
 }
