Make sure to call ExecuteOnGameThread instead of Async, for OSSv2 operation steps using game thread execution policy, to avoid calling user callbacks during GC.

diff --git a/Engine/Plugins/Online/OnlineServices/Source/OnlineServicesCommon/Public/Online/OnlineAsyncOp.h b/Engine/Plugins/Online/OnlineServices/Source/OnlineServicesCommon/Public/Online/OnlineAsyncOp.h
--- a/Engine/Plugins/Online/OnlineServices/Source/OnlineServicesCommon/Public/Online/OnlineAsyncOp.h
+++ b/Engine/Plugins/Online/OnlineServices/Source/OnlineServicesCommon/Public/Online/OnlineAsyncOp.h
@@ -1065,12 +1065,12 @@ public:
 			}
 			else
 			{
-				Async(EAsyncExecution::TaskGraphMainThread, MoveTemp(Callable));
+				ExecuteOnGameThread(OpType::Name, MoveTemp(Callable));
 			}
 			break;
 
 		case EOnlineAsyncExecutionPolicy::RunOnNextTick:
-			Async(EAsyncExecution::TaskGraphMainThread, MoveTemp(Callable));
+			ExecuteOnGameThread(OpType::Name, MoveTemp(Callable));
 			break;
 
 		case EOnlineAsyncExecutionPolicy::RunOnThreadPool:
