Add static check to make sure when define error category in different system it won't use the same category value;
Make sure the subclass of OnlineErrorDetails can use default ctor of base class and parse error details later.

diff --git a/Engine/Plugins/Online/OnlineServices/Source/OnlineServicesInterface/Public/Online/OnlineError.h b/Engine/Plugins/Online/OnlineServices/Source/OnlineServicesInterface/Public/Online/OnlineError.h
--- a/Engine/Plugins/Online/OnlineServices/Source/OnlineServicesInterface/Public/Online/OnlineError.h
+++ b/Engine/Plugins/Online/OnlineServices/Source/OnlineServicesInterface/Public/Online/OnlineError.h
@@ -29,7 +29,6 @@ public:
 	{
 	}
 
-
 	virtual ~FOnlineErrorDetails() {}
 
 	virtual FString GetFriendlyErrorCode(const FOnlineError&) const override
@@ -48,6 +47,10 @@ public:
 	}
 
 protected:
+	FOnlineErrorDetails()
+	{
+	}
+
 	FString FriendlyErrorCode;
 	FString LogString;
 	FText Text;
diff --git a/Engine/Plugins/Online/OnlineServices/Source/OnlineServicesInterface/Public/Online/OnlineErrorCode.h b/Engine/Plugins/Online/OnlineServices/Source/OnlineServicesInterface/Public/Online/OnlineErrorCode.h
--- a/Engine/Plugins/Online/OnlineServices/Source/OnlineServicesInterface/Public/Online/OnlineErrorCode.h
+++ b/Engine/Plugins/Online/OnlineServices/Source/OnlineServicesInterface/Public/Online/OnlineErrorCode.h
@@ -43,6 +43,7 @@ ONLINESERVICESINTERFACE_API uint64 ErrorCodeValue(ErrorCodeType ErrorCode);
 
 // define an error category
 #define UE_ONLINE_ERROR_CATEGORY(Name, InSystem, Value, Description) \
+namespace ErrorCode { namespace Category { static constexpr int Category_Value_##Value##_Is_Taken_In_##InSystem = 0; } }\
 namespace ErrorCode { namespace Category { static constexpr int Name = Value; } }\
 namespace ErrorCode { namespace Category { static constexpr int Name##_System = UE::Online::Errors::ErrorCode::System::InSystem; } }  
 
