Add support to return OnlineErrorDetails from OnlineError.

diff --git a/Engine/Plugins/Online/OnlineServices/Source/OnlineServicesInterface/Public/Online/OnlineError.h b/Engine/Plugins/Online/OnlineServices/Source/OnlineServicesInterface/Public/Online/OnlineError.h
--- a/Engine/Plugins/Online/OnlineServices/Source/OnlineServicesInterface/Public/Online/OnlineError.h
+++ b/Engine/Plugins/Online/OnlineServices/Source/OnlineServicesInterface/Public/Online/OnlineError.h
@@ -151,7 +151,12 @@ public:
 
 	const FOnlineError* GetInner() const
 	{
-		return Inner.IsValid() ? Inner.Get() : nullptr;
+		return Inner.Get();
+	}
+
+	const IOnlineErrorDetails* GetErrorDetails() const
+	{
+		return Details.Get();
 	}
 
 	ErrorCodeType GetSystem() const
