Update LAN sessions & Join Session Tests
- Cleaner API for TryHostLANSession to utilize TValueOrError.
- Simplify session creation flow with a new AddLANSession function.
- Keep track of the hosting LAN session instead of a basic counter, as only 1 can be hosted at a time.
- Add the hosting LAN session to search results if one is hosted, as the LAN beacon must be stopped to execute a search.
- Find & remove the promise while also emplacing a value to reduce duplicated TMap look-ups.
- Update join sessions test to ensure only 1 hosted session occurs.
- Update join sessions test to remove unnecessary 6-second wait time.

diff --git a/Engine/Plugins/Online/OnlineServices/Source/OnlineServicesCommonEngineUtils/Private/Online/SessionsLAN.cpp b/Engine/Plugins/Online/OnlineServices/Source/OnlineServicesCommonEngineUtils/Private/Online/SessionsLAN.cpp
--- a/Engine/Plugins/Online/OnlineServices/Source/OnlineServicesCommonEngineUtils/Private/Online/SessionsLAN.cpp
+++ b/Engine/Plugins/Online/OnlineServices/Source/OnlineServicesCommonEngineUtils/Private/Online/SessionsLAN.cpp
@@ -114,6 +114,15 @@ FSessionsLAN::FSessionsLAN(FOnlineServicesCommon& InServices)
 
 }
 
+FSessionsLAN::~FSessionsLAN()
+{
+	// Stop hosted LAN session if one is running
+	if (HostedLANSession)
+	{
+		StopLANSession();
+	}
+}
+
 void FSessionsLAN::Tick(float DeltaSeconds)
 {
 	if (LANSessionManager->GetBeaconState() > ELanBeaconState::NotUsingLanBeacon)
@@ -132,25 +141,20 @@ TFuture<TOnlineResult<FCreateSession>> FSessionsLAN::CreateSessionImpl(const FCr
 	// We'll only host on the LAN beacon for public sessions
 	if (Params.SessionSettings.JoinPolicy == ESessionJoinPolicy::Public)
 	{
-		if (TOptional<FOnlineError> TryHostLANSessionResult = TryHostLANSession())
+		FHostSessionResult TryHostLANSessionResult = TryHostLANSession();
+		if (TryHostLANSessionResult.HasError())
 		{
-			Promise.EmplaceValue(TryHostLANSessionResult.GetValue());
+			Promise.EmplaceValue(TryHostLANSessionResult.StealError());
 			return Future;
 		}
-	}
-
-	TSharedRef<FSessionLAN> NewSessionLANRef = MakeShared<FSessionLAN>();
-	NewSessionLANRef->OwnerAccountId = Params.LocalAccountId;
-	NewSessionLANRef->SessionInfo.bAllowSanctionedPlayers = Params.bAllowSanctionedPlayers;
-	NewSessionLANRef->SessionInfo.bAntiCheatProtected = Params.bAntiCheatProtected;
-	NewSessionLANRef->SessionInfo.bIsDedicatedServerSession = IsRunningDedicatedServer();
-	NewSessionLANRef->SessionInfo.bIsLANSession = true;
-	NewSessionLANRef->SessionInfo.SessionId = FOnlineSessionIdRegistryLAN::GetChecked(Services.GetServicesProvider()).GetNextSessionId();
-	NewSessionLANRef->SessionInfo.SessionIdOverride = Params.SessionIdOverride;
-	NewSessionLANRef->SessionSettings = Params.SessionSettings;
 
-	// We save the local object for the session, and set up the appropriate references
-	AddSessionWithReferences(NewSessionLANRef, Params.SessionName, Params.LocalAccountId, Params.bPresenceEnabled);
+		UE_LOG(LogOnlineServices, Verbose, TEXT("Hosting LAN session named '%s'"), *(Params.SessionName.ToString()));
+		HostedLANSession = AddLANSession(Params);
+	}
+	else
+	{
+		AddLANSession(Params);
+	}
 
 	Promise.EmplaceValue(FCreateSession::Result{ });
 
@@ -176,15 +180,24 @@ TFuture<TOnlineResult<FUpdateSessionSettings>> FSessionsLAN::UpdateSessionSettin
 		// There is no FriendsLAN interface at the time of this implementation, hence the binary behavior (Public/Non-Public)
 		if (SessionSettings.JoinPolicy != ESessionJoinPolicy::Public && Params.Mutations.JoinPolicy.GetValue() == ESessionJoinPolicy::Public)
 		{
-			if (TOptional<FOnlineError> TryHostLANSessionResult = TryHostLANSession())
+			FHostSessionResult TryHostLANSessionResult = TryHostLANSession();
+			if (TryHostLANSessionResult.HasError())
 			{
-				Promise.EmplaceValue(TryHostLANSessionResult.GetValue());
+				Promise.EmplaceValue(TryHostLANSessionResult.StealError());
 				return Future;
 			}
+
+			// This session already exists, we are moving it to be a hosted session
+			// So we only need to update our hosting session information
+			UE_LOG(LogOnlineServices, Verbose, TEXT("Transitioning session %s to a hosted LAN session."), *Params.SessionName.ToString());
+			HostedLANSession = StaticCastSharedRef<FSessionLAN>(FoundSession).ToSharedPtr();
 		}
 		else if (SessionSettings.JoinPolicy == ESessionJoinPolicy::Public && Params.Mutations.JoinPolicy.GetValue() != ESessionJoinPolicy::Public)
 		{
-			StopLANSession();
+			// Stop hosting the LAN session, since it's no longer public
+			LANSessionManager->StopLANSession();
+			HostedLANSession.Reset();
+			UE_LOG(LogOnlineServices, Verbose, TEXT("No longer hosting LAN session %s as it's no longer publicly accessible."), *Params.SessionName.ToString());
 		}
 	}
 
@@ -211,6 +224,13 @@ TFuture<TOnlineResult<FFindSessions>> FSessionsLAN::FindSessionsImpl(const FFind
 	SearchResultsUserMap.FindOrAdd(Params.LocalAccountId).Reset();
 	CurrentSessionSearchPromisesUserMap.Emplace(Params.LocalAccountId, MoveTemp(Promise));
 
+	// If we were hosting public sessions before the search, add it to the search results
+	if (HostedLANSession)
+	{
+		UE_LOG(LogOnlineServices, Verbose, TEXT("[FSessionsLAN::FindSessionsImpl] Adding local hosted LAN session with ID %s to search results."), *ToString(HostedLANSession->GetSessionId()));
+		SearchResultsUserMap.FindOrAdd(Params.LocalAccountId).Add(HostedLANSession->GetSessionId());
+	}
+
 	FindLANSessions(Params.LocalAccountId);
 
 	return Future;
@@ -281,10 +301,8 @@ TFuture<TOnlineResult<FLeaveSession>> FSessionsLAN::LeaveSessionImpl(const FLeav
 
 /** LANSessionManager methods */
 
-TOptional<FOnlineError> FSessionsLAN::TryHostLANSession()
+FSessionsLAN::FHostSessionResult FSessionsLAN::TryHostLANSession()
 {
-	TOptional<FOnlineError> Result;
-
 	// The LAN Beacon can only broadcast one session at a time
 	if (LANSessionManager->GetBeaconState() == ELanBeaconState::NotUsingLanBeacon)
 	{
@@ -294,7 +312,7 @@ TOptional<FOnlineError> FSessionsLAN::TryHostLANSession()
 		{
 			UE_LOG(LogOnlineServices, VeryVerbose, TEXT("[FSessionsLAN::TryHostLANSession] LAN Beacon hosting..."));
 
-			PublicSessionsHosted++;
+			return MakeValue();
 		}
 		else
 		{
@@ -302,17 +320,32 @@ TOptional<FOnlineError> FSessionsLAN::TryHostLANSession()
 
 			LANSessionManager->StopLANSession();
 
-			Result.Emplace(Errors::RequestFailure());
+			return MakeError(Errors::RequestFailure());
 		}
 	}
 	else
 	{
-		UE_LOG(LogOnlineServices, VeryVerbose, TEXT("[FSessionsLAN::TryHostLANSession] LAN Beacon already in use!"));
-
-		Result.Emplace(Errors::AlreadyPending());
+		UE_LOG(LogOnlineServices, Warning, TEXT("[FSessionsLAN::TryHostLANSession] LAN Beacon already in use!"));
+		return MakeError(Errors::AlreadyPending());
 	}
+}
 
-	return Result;
+TSharedRef<FSessionLAN> FSessionsLAN::AddLANSession(const FCreateSession::Params& Params)
+{
+	TSharedRef<FSessionLAN> NewSessionLANRef = MakeShared<FSessionLAN>();
+	NewSessionLANRef->OwnerAccountId = Params.LocalAccountId;
+	NewSessionLANRef->SessionInfo.bAllowSanctionedPlayers = Params.bAllowSanctionedPlayers;
+	NewSessionLANRef->SessionInfo.bAntiCheatProtected = Params.bAntiCheatProtected;
+	NewSessionLANRef->SessionInfo.bIsDedicatedServerSession = IsRunningDedicatedServer();
+	NewSessionLANRef->SessionInfo.bIsLANSession = true;
+	NewSessionLANRef->SessionInfo.SessionId = FOnlineSessionIdRegistryLAN::GetChecked(Services.GetServicesProvider()).GetNextSessionId();
+	NewSessionLANRef->SessionInfo.SessionIdOverride = Params.SessionIdOverride;
+	NewSessionLANRef->SessionSettings = Params.SessionSettings;
+
+	// We save the local object for the session, and set up the appropriate references
+	AddSessionWithReferences(NewSessionLANRef, Params.SessionName, Params.LocalAccountId, Params.bPresenceEnabled);
+
+	return NewSessionLANRef;
 }
 
 void FSessionsLAN::FindLANSessions(const FAccountId& LocalAccountId)
@@ -336,11 +369,10 @@ void FSessionsLAN::FindLANSessions(const FAccountId& LocalAccountId)
 		}
 
 		// Trigger the delegate as having failed
-		CurrentSessionSearchPromisesUserMap.FindChecked(LocalAccountId).EmplaceValue(Errors::RequestFailure());
-		CurrentSessionSearchPromisesUserMap.Remove(LocalAccountId);
+		CurrentSessionSearchPromisesUserMap.FindAndRemoveChecked(LocalAccountId).EmplaceValue(Errors::RequestFailure());
 
 		// If we were hosting public sessions before the search, we'll return the beacon to that state
-		if (PublicSessionsHosted > 0)
+		if (HostedLANSession)
 		{
 			TryHostLANSession();
 		}
@@ -353,14 +385,9 @@ void FSessionsLAN::FindLANSessions(const FAccountId& LocalAccountId)
 
 void FSessionsLAN::StopLANSession()
 {
-	check(PublicSessionsHosted > 0);
-
-	PublicSessionsHosted--;
-
-	if (PublicSessionsHosted == 0)
-	{
-		LANSessionManager->StopLANSession();
-	}
+	check(HostedLANSession);
+	HostedLANSession.Reset();
+	LANSessionManager->StopLANSession();
 }
 
 void FSessionsLAN::OnValidQueryPacketReceived(uint8* PacketData, int32 PacketLength, uint64 ClientNonce)
@@ -409,15 +436,13 @@ void FSessionsLAN::OnLANSearchTimeout(const FAccountId LocalAccountId)
 		LANSessionManager->StopLANSession();
 	}
 
-	UE_LOG(LogOnlineServices, Verbose, TEXT("[FSessionsLAN::OnLANSearchTimeout] %d sessions found!"), SearchResultsUserMap.FindChecked(LocalAccountId).Num());
-
-	TArray<FOnlineSessionId>& FoundSessionIds = SearchResultsUserMap.FindChecked(LocalAccountId);
+	const TArray<FOnlineSessionId>& FoundSessionIds = SearchResultsUserMap.FindChecked(LocalAccountId);
+	UE_LOG(LogOnlineServices, Verbose, TEXT("[FSessionsLAN::OnLANSearchTimeout] %d sessions found!"), FoundSessionIds.Num());
 
-	CurrentSessionSearchPromisesUserMap.FindChecked(LocalAccountId).EmplaceValue(FFindSessions::Result{ FoundSessionIds });
-	CurrentSessionSearchPromisesUserMap.Remove(LocalAccountId); 
+	CurrentSessionSearchPromisesUserMap.FindAndRemoveChecked(LocalAccountId).EmplaceValue(FFindSessions::Result{ FoundSessionIds });
 
 	// If we were hosting public sessions before the search, we'll return the beacon to that state
-	if (PublicSessionsHosted > 0)
+	if (HostedLANSession)
 	{
 		TryHostLANSession();
 	}
diff --git a/Engine/Plugins/Online/OnlineServices/Source/OnlineServicesCommonEngineUtils/Public/Online/SessionsLAN.h b/Engine/Plugins/Online/OnlineServices/Source/OnlineServicesCommonEngineUtils/Public/Online/SessionsLAN.h
--- a/Engine/Plugins/Online/OnlineServices/Source/OnlineServicesCommonEngineUtils/Public/Online/SessionsLAN.h
+++ b/Engine/Plugins/Online/OnlineServices/Source/OnlineServicesCommonEngineUtils/Public/Online/SessionsLAN.h
@@ -2,6 +2,7 @@
 #pragma once
 
 #include "Online/SessionsCommon.h"
+#include "Templates/ValueOrError.h"
 
 class FInternetAddr;
 class FLANSession;
@@ -54,7 +55,7 @@ public:
 	using Super = FSessionsCommon;
 
 	FSessionsLAN(FOnlineServicesCommon& InServices);
-	virtual ~FSessionsLAN() = default;
+	virtual ~FSessionsLAN() override;
 
 	virtual void Tick(float DeltaSeconds) override;
 
@@ -66,7 +67,10 @@ public:
 	virtual TFuture<TOnlineResult<FJoinSession>> JoinSessionImpl(const FJoinSession::Params& Params) override;
 
 protected:
-	TOptional<FOnlineError> TryHostLANSession();
+	using FHostSessionResult = TValueOrError<void, FOnlineError>;
+
+	FHostSessionResult TryHostLANSession();
+	TSharedRef<FSessionLAN> AddLANSession(const FCreateSession::Params& Params);
 	void FindLANSessions(const FAccountId& LocalAccountId);
 	void StopLANSession();
 	void OnValidQueryPacketReceived(uint8* PacketData, int32 PacketLength, uint64 ClientNonce);
@@ -75,9 +79,7 @@ protected:
 	virtual void AppendSessionToPacket(FNboSerializeToBuffer& Packet, const FSessionLAN& Session) = 0;
 	virtual void ReadSessionFromPacket(FNboSerializeFromBuffer& Packet, FSessionLAN& Session) = 0;
 
-protected:
-	uint32 PublicSessionsHosted = 0;
-
+	TSharedPtr<FSessionLAN> HostedLANSession = nullptr;
 	TSharedRef<FLANSession> LANSessionManager;
 };
 
diff --git a/Engine/Source/Programs/Online/OnlineTestsCore/Source/Private/Tests/Sessions/JoinSessionTests.cpp b/Engine/Source/Programs/Online/OnlineTestsCore/Source/Private/Tests/Sessions/JoinSessionTests.cpp
--- a/Engine/Source/Programs/Online/OnlineTestsCore/Source/Private/Tests/Sessions/JoinSessionTests.cpp
+++ b/Engine/Source/Programs/Online/OnlineTestsCore/Source/Private/Tests/Sessions/JoinSessionTests.cpp
@@ -127,6 +127,10 @@ SESSIONS_TEST_CASE("If I call JoinSession with an name already in use, I get an
 	SecondCreateSessionHelperParams.OpParams->SessionSettings.SchemaName = TEXT("SessionNameInUseSchemaName1");
 	SecondCreateSessionHelperParams.OpParams->SessionSettings.NumMaxConnections = 4;
 
+	// Friends Only for the second session, because otherwise it defaults to public.
+	// For LAN sessions, we will try to host public ones and only 1 can be hosted at a time, which should be occupied by the first session.
+	SecondCreateSessionHelperParams.OpParams->SessionSettings.JoinPolicy = ESessionJoinPolicy::FriendsOnly;
+
 	FAddSessionMember::Params OpFirstAddParams;
 	FAddSessionMemberHelper::FHelperParams FirstAddSessionMemberHelperParams;
 	FirstAddSessionMemberHelperParams.OpParams = &OpFirstAddParams;
@@ -187,7 +191,6 @@ SESSIONS_TEST_CASE("If I call JoinSession with an name already in use, I get an
 	LoginPipeline
 		.EmplaceStep<FCreateSessionHelper>(MoveTemp(FirstCreateSessionHelperParams))
 		.EmplaceStep<FCreateSessionHelper>(MoveTemp(SecondCreateSessionHelperParams))
-		.EmplaceStep<FTickForTime>(FTimespan::FromMilliseconds(6000))
 		.EmplaceStep<FAddSessionMemberHelper>(MoveTemp(FirstAddSessionMemberHelperParams))
 		.EmplaceStep<FAddSessionMemberHelper>(MoveTemp(SecondAddSessionMemberHelperParams))
 		.EmplaceLambda([&FirstFindSessionsHelperParams](SubsystemType OnlineSubsystem)
