OnineService async op changes
part1

diff --git a/Engine/Plugins/Online/OnlineServices/Source/OnlineServicesCommon/Private/Online/OnlineSyntaxTests.cpp b/Engine/Plugins/Online/OnlineServices/Source/OnlineServicesCommon/Private/Online/OnlineSyntaxTests.cpp
--- a/Engine/Plugins/Online/OnlineServices/Source/OnlineServicesCommon/Private/Online/OnlineSyntaxTests.cpp
+++ b/Engine/Plugins/Online/OnlineServices/Source/OnlineServicesCommon/Private/Online/OnlineSyntaxTests.cpp
@@ -196,6 +196,10 @@ inline void TestAsyncOpContinuationSyntax(FOnlineServicesCommon& Services)
 		.ForEachN(2, [](TOnlineAsyncOp<FTestOp>&, TArrayView<int> Value) { return MakeFulfilledPromise<TArray<int>>(Value).GetFuture(); })
 		.Then([](TOnlineAsyncOp<FTestOp>&, const TArray<int>&) { return 1; })
 		.Then([](TOnlineAsyncOp<FTestOp>&, int) {})
+		.Then([](TOnlineAsyncOp<FTestOp>&) { return TContinuationResult<int>::Complete(1); })
+		.Then([](TOnlineAsyncOp<FTestOp>&, int) { return MakeFulfilledPromise<TContinuationResult<int>>(TContinuationResult<int>::Complete(1)).GetFuture(); })
+		.Then([](TOnlineAsyncOp<FTestOp>&, int) { return MakeFulfilledPromise<TContinuationResult<void>>(TContinuationResult<void>::Complete()).GetFuture(); })
+		.Then([](TOnlineAsyncOp<FTestOp>&) { return TContinuationResult<void>::Complete(); })
 		.Enqueue(Services.GetParallelQueue());
 	Op.GetHandle();
 }
diff --git a/Engine/Plugins/Online/OnlineServices/Source/OnlineServicesCommon/Public/Online/OnlineAsyncOp.h b/Engine/Plugins/Online/OnlineServices/Source/OnlineServicesCommon/Public/Online/OnlineAsyncOp.h
--- a/Engine/Plugins/Online/OnlineServices/Source/OnlineServicesCommon/Public/Online/OnlineAsyncOp.h
+++ b/Engine/Plugins/Online/OnlineServices/Source/OnlineServicesCommon/Public/Online/OnlineAsyncOp.h
@@ -73,6 +73,60 @@ public:
 	}
 };
 
+template <typename TResultType>
+struct TContinuationResult
+{
+public:
+	static TContinuationResult Repeat() { return TContinuationResult<TResultType>(); }
+	static TContinuationResult Complete(TResultType&& InResult) { return TContinuationResult(MoveTempIfPossible(InResult)); }
+	bool IsComplete() const { return Result.IsSet(); }
+	TResultType& GetResult() { return Result.GetValue(); }
+	const TResultType& GetResult() const { return Result.GetValue(); }
+private:
+
+	TContinuationResult()
+	{
+	}
+	TContinuationResult(TResultType&& InResult) : Result(InResult)
+	{
+	}
+	TOptional<TResultType> Result;
+};
+
+template <>
+struct TContinuationResult<void>
+{
+public:
+	static TContinuationResult Repeat() { return TContinuationResult<void>(); }
+	static TContinuationResult Complete() { return TContinuationResult(FIsCompleteTag()); }
+	bool IsComplete() const { return bIsComplete; }
+private:
+	struct FIsCompleteTag {};
+
+	TContinuationResult()
+	{
+	}
+	TContinuationResult(const FIsCompleteTag&) : bIsComplete(true)
+	{
+	}
+	bool bIsComplete = false;
+};
+
+template <typename T>
+struct TContinuationResultType
+{
+	using ResultType = T;
+};
+
+template <typename T>
+struct TContinuationResultType<TContinuationResult<T>>
+{
+	using ResultType = T;
+};
+
+template <typename T>
+using TContinuationResultType_T = typename TContinuationResultType<T>::ResultType;
+
 template <typename T>
 FString ToLogString(const TDefaultErrorResultInternal<T>& Result)
 {
@@ -533,6 +587,319 @@ private:
 	TUniqueFunction<void()> ExecFunction;
 };
 
+template <typename ResultType>
+class TStep<TContinuationResult<ResultType>> : public IStep
+{
+	public:
+	TStep(FOnlineAsyncExecutionPolicy&& InExecutionPolicy)
+		: ExecutionPolicy(MoveTemp(InExecutionPolicy))
+	{
+	}
+
+	~TStep()
+	{
+		if (bResultSet)
+		{
+			DestructItem(Result.GetTypedPtr());
+		}
+	}
+
+	template <typename OpType, typename LastResultType, typename CallableType>
+	void SetExecFunction(TOnlineAsyncOp<OpType>& InOperation, LastResultType&& InLastResult, CallableType&& InCallable)
+	{
+		ExecFunction = [this, WeakOperation = TWeakPtr<TOnlineAsyncOp<OpType>>(InOperation.AsShared()), &LastResult = InLastResult, Callable = MoveTemp(InCallable)]() mutable
+			{
+				TSharedPtr<TOnlineAsyncOp<OpType>> PinnedOperation = WeakOperation.Pin();
+				if (PinnedOperation)
+				{
+					if constexpr (TOnlineAsyncOpCallableTraits<CallableType>::bRequiresPromise)
+					{
+						TPromise<TContinuationResult<ResultType>> Promise;
+						// set promise continuation before calling the callable so that we will complete the step as soon as the value is set
+						Promise.GetFuture()
+							.Next([this, WeakOperation](const TContinuationResult<ResultType>& Value)
+								{
+									TSharedPtr<TOnlineAsyncOp<OpType>> PinnedOperation2 = WeakOperation.Pin();
+
+									if (PinnedOperation2)
+									{
+										if (!Value.IsComplete())
+										{
+											PinnedOperation2->ExecuteRepeatStep();
+										}
+										else
+										{
+											EmplaceResult(Value);
+											PinnedOperation2->ExecuteNextStep();
+										}
+									}
+								});
+						Callable(*PinnedOperation, MoveTempIfPossible(LastResult), MoveTemp(Promise));
+					}
+					else if constexpr (TOnlineAsyncOpCallableTraits<CallableType>::bAsyncResult)
+					{
+						Callable(*PinnedOperation, LastResult)
+							.Next([this, WeakOperation](const TContinuationResult<ResultType>& Value)
+								{
+
+									TSharedPtr<TOnlineAsyncOp<OpType>> PinnedOperation2 = WeakOperation.Pin();
+									if (PinnedOperation2)
+									{
+										if (!Value.IsComplete())
+										{
+											PinnedOperation2->ExecuteRepeatStep();
+										}
+										else
+										{
+											EmplaceResult(Value.GetResult());
+											PinnedOperation2->ExecuteNextStep();
+										}
+									}
+								});
+
+					}
+					else
+					{
+						const TContinuationResult<ResultType>& Value = Callable(*PinnedOperation, LastResult);
+						if (!Value.IsComplete())
+						{
+							PinnedOperation->ExecuteRepeatStep();
+						}
+						else
+						{
+							EmplaceResult(Value.GetResult());
+							PinnedOperation->ExecuteNextStep();
+						}
+					}
+				}
+			};
+	}
+
+	template <typename OpType, typename CallableType>
+	void SetExecFunction(TOnlineAsyncOp<OpType>& InOperation, CallableType&& InCallable)
+	{
+		ExecFunction = [this, WeakOperation = TWeakPtr<TOnlineAsyncOp<OpType>>(InOperation.AsShared()), Callable = MoveTemp(InCallable)]() mutable
+			{
+				TSharedPtr<TOnlineAsyncOp<OpType>> PinnedOperation = WeakOperation.Pin();
+				if (PinnedOperation)
+				{
+					if constexpr (TOnlineAsyncOpCallableTraits<CallableType>::bRequiresPromise)
+					{
+						TPromise<TContinuationResult<ResultType>> Promise;
+						// set promise continuation before calling the callable so that we will complete the step as soon as the value is set
+						Promise.GetFuture()
+							.Next([this, WeakOperation](const TContinuationResult<ResultType>& Value)
+								{
+									TSharedPtr<TOnlineAsyncOp<OpType>> PinnedOperation2 = WeakOperation.Pin();
+									if (PinnedOperation2)
+									{
+										if (!Value.IsComplete())
+										{
+											PinnedOperation2->ExecuteRepeatStep();
+										}
+										else
+										{
+											EmplaceResult(Value);
+											PinnedOperation2->ExecuteNextStep();
+										}
+									}
+								});
+						Callable(*PinnedOperation, MoveTemp(Promise));
+					}
+					else if constexpr (TOnlineAsyncOpCallableTraits<CallableType>::bAsyncResult)
+					{
+
+						Callable(*PinnedOperation)
+							.Next([this, WeakOperation](const TContinuationResult<ResultType>& Value)
+								{
+
+									TSharedPtr<TOnlineAsyncOp<OpType>> PinnedOperation2 = WeakOperation.Pin();
+									if (PinnedOperation2)
+									{
+										if (!Value.IsComplete())
+										{
+											PinnedOperation2->ExecuteRepeatStep();
+										}
+										else
+										{
+											EmplaceResult(Value);
+											PinnedOperation2->ExecuteNextStep();
+										}
+									}
+
+								});
+
+					}
+					else
+					{
+						const TContinuationResult<ResultType>& Value = Callable(*PinnedOperation);
+						if (!Value.IsComplete())
+						{
+							PinnedOperation->ExecuteRepeatStep();
+						}
+						else
+						{
+							EmplaceResult(Value.GetResult());
+							PinnedOperation->ExecuteNextStep();
+						}
+					}
+				}
+			};
+	}
+
+	virtual const FOnlineAsyncExecutionPolicy& GetExecutionPolicy() const override
+	{
+		return ExecutionPolicy;
+	}
+
+	virtual void Execute() override
+	{
+		// Repeating continuations will loop until finished. Not safe to run on game thread.
+		check(ExecutionPolicy.GetExecutionPolicy() != EOnlineAsyncExecutionPolicy::RunOnGameThread)
+			check(ExecFunction)
+			ExecFunction();
+	}
+
+
+	ResultType& GetResultRef()
+	{
+		return *Result.GetTypedPtr();
+	}
+
+private:
+	template<typename... ArgTypes>
+	void EmplaceResult(ArgTypes&&... Args)
+	{
+		check(!bResultSet)
+			new(Result.GetTypedPtr()) ResultType(Forward<ArgTypes>(Args)...);
+		bResultSet = true;
+	}
+
+	FOnlineAsyncExecutionPolicy ExecutionPolicy;
+	TUniqueFunction<void()> ExecFunction;
+	TTypeCompatibleBytes<ResultType> Result;
+	bool bResultSet = false;
+};
+
+template <>
+class TStep<TContinuationResult<void>> : public IStep
+{
+public:
+	TStep(FOnlineAsyncExecutionPolicy&& InExecutionPolicy)
+		: ExecutionPolicy(MoveTemp(InExecutionPolicy))
+	{
+	}
+
+	template <typename OpType, typename LastResultType, typename CallableType>
+	void SetExecFunction(TOnlineAsyncOp<OpType>& InOperation, LastResultType&& InLastResult, CallableType&& InCallable)
+	{
+		ExecFunction = [this, WeakOperation = TWeakPtr<TOnlineAsyncOp<OpType>>(InOperation.AsShared()), &LastResult = InLastResult, Callable = MoveTemp(InCallable)]() mutable
+			{
+				TSharedPtr<TOnlineAsyncOp<OpType>> PinnedOperation = WeakOperation.Pin();
+				if (PinnedOperation)
+				{
+					if constexpr (TOnlineAsyncOpCallableTraits<CallableType>::bAsyncResult)
+					{
+						Callable(*PinnedOperation, LastResult)
+							.Next([this, WeakOperation](const TContinuationResult<void>& Value)
+							{
+
+								TSharedPtr<TOnlineAsyncOp<OpType>> PinnedOperation2 = WeakOperation.Pin();
+								if (PinnedOperation2)
+								{
+									if (!Value.IsComplete())
+									{
+										PinnedOperation2->ExecuteRepeatStep();
+									}
+									else
+									{
+										PinnedOperation2->ExecuteNextStep();
+									}
+								}
+							});
+
+					}
+					else
+					{
+						const TContinuationResult<void>& Value = Callable(*PinnedOperation, LastResult);
+						if (!Value.IsComplete())
+						{
+							PinnedOperation->ExecuteRepeatStep();
+						}
+						else
+						{
+							PinnedOperation->ExecuteNextStep();
+						}
+					}
+				}
+			};
+	}
+
+	template <typename OpType, typename CallableType>
+	void SetExecFunction(TOnlineAsyncOp<OpType>& InOperation, CallableType&& InCallable)
+	{
+		ExecFunction = [this, WeakOperation = TWeakPtr<TOnlineAsyncOp<OpType>>(InOperation.AsShared()), Callable = MoveTemp(InCallable)]() mutable
+			{
+				TSharedPtr<TOnlineAsyncOp<OpType>> PinnedOperation = WeakOperation.Pin();
+				if (PinnedOperation)
+				{
+					if constexpr (TOnlineAsyncOpCallableTraits<CallableType>::bAsyncResult)
+					{
+						Callable(*PinnedOperation)
+							.Next([this, WeakOperation](const TContinuationResult<void>& Value)
+							{
+
+								TSharedPtr<TOnlineAsyncOp<OpType>> PinnedOperation2 = WeakOperation.Pin();
+								if (PinnedOperation2)
+								{
+									if (!Value.IsComplete())
+									{
+										PinnedOperation2->ExecuteRepeatStep();
+									}
+									else
+									{
+										PinnedOperation2->ExecuteNextStep();
+									}
+								}
+
+							});
+					}
+					else
+					{
+						const TContinuationResult<void>& Value = Callable(*PinnedOperation);
+						if (!Value.IsComplete())
+						{
+							PinnedOperation->ExecuteRepeatStep();
+						}
+						else
+						{
+							PinnedOperation->ExecuteNextStep();
+						}
+					}
+				}
+			};
+	}
+
+	virtual const FOnlineAsyncExecutionPolicy& GetExecutionPolicy() const override
+	{
+		return ExecutionPolicy;
+	}
+
+	virtual void Execute() override
+	{
+		// Repeating continuations will loop until finished. Not safe to run on game thread.
+		check(ExecutionPolicy.GetExecutionPolicy() != EOnlineAsyncExecutionPolicy::RunOnGameThread)
+		check(ExecFunction)
+		ExecFunction();
+	}
+
+private:
+	FOnlineAsyncExecutionPolicy ExecutionPolicy;
+	TUniqueFunction<void()> ExecFunction;
+};
+
+
+
 template <typename ResultElementType>
 class TForEachStepBase : public IStep
 {
@@ -795,6 +1162,10 @@ public:
 	//     AsyncOp is a FOnlineAsyncOp& or TOnlineAsyncOp<OpType>&
 	//     LastResult is a LastResultType or const LastResultType&
 	//     ResultType is any non-void type
+	// For either ResultType can be TContinuationResult<T>. 
+	//   Retuning TContinuationResult<T>::Repeat() will cause the step to repeat.
+	//   Returning TContinuationResult<T>::Complete(T) will complete the step and allow the task to proceed to the next.
+	//	 The type T in TContinuationResult<T>::Complete(T) will be the LastResult type for a subsequent step.
 	template <typename CallableType>
 	auto Then(CallableType&& Callable, FOnlineAsyncExecutionPolicy ExecutionPolicy = FOnlineAsyncExecutionPolicy::RunOnGameThread());
 
@@ -810,6 +1181,8 @@ public:
 	//     AsyncOp is a FOnlineAsyncOp& or TOnlineAsyncOp<OpType>&
 	//     LastResultElement is a LastResultType or const LastResultType&
 	//     ResultType is any non-void type
+
+
 	template <typename CallableType> // previous value was TArray
 	auto ForEach(CallableType&& Callable, FOnlineAsyncExecutionPolicy ExecutionPolicy = FOnlineAsyncExecutionPolicy::RunOnGameThread());
 
@@ -848,18 +1221,19 @@ public:
 	//     ResultType is any non-void type
 	template <typename CallableType>
 	auto Then(CallableType&& Callable, FOnlineAsyncExecutionPolicy ExecutionPolicy = FOnlineAsyncExecutionPolicy::RunOnGameThread());
+
 };
 
 
 /* Private */ }
 
 template <typename OpType, typename T>
-class TOnlineChainableAsyncOp : public Private::TOnlineAsyncOpBase<TOnlineChainableAsyncOp<OpType, T>, OpType, T>
+class TOnlineChainableAsyncOp : public Private::TOnlineAsyncOpBase<TOnlineChainableAsyncOp<OpType, T>, OpType, TContinuationResultType_T<T>>
 {
 public:
-	using Super = Private::TOnlineAsyncOpBase<TOnlineChainableAsyncOp<OpType, T>, OpType, T>;
+	using Super = Private::TOnlineAsyncOpBase<TOnlineChainableAsyncOp<OpType, T>, OpType, TContinuationResultType_T<T>>;
 
-	TOnlineChainableAsyncOp(TOnlineAsyncOp<OpType>& InOwningOperation, std::enable_if_t<!std::is_same_v<T, void>, T>& InLastResult)
+	TOnlineChainableAsyncOp(TOnlineAsyncOp<OpType>& InOwningOperation, std::enable_if_t<!std::is_same_v<TContinuationResultType_T<T>, void>, TContinuationResultType_T<T>>& InLastResult)
 		: Super(InLastResult)
 		, OwningOperation(InOwningOperation)
 	{
@@ -1076,6 +1450,26 @@ public:
 		}
 	}
 
+	void ExecuteRepeatStep()
+	{
+		if (!IsComplete())
+		{
+			const int StepToExecute = NextStep == 0 ? 0 : NextStep - 1;
+			if (StepToExecute < Steps.Num())
+			{
+				Execute(Steps[StepToExecute]->GetExecutionPolicy(),
+					[this, StepToExecute, WeakThis = TWeakPtr<TOnlineAsyncOp<OpType>>(this->AsShared())]()
+					{
+						TSharedPtr<TOnlineAsyncOp<OpType>> PinnedThis = WeakThis.Pin();
+						if (PinnedThis)
+						{
+							Steps[StepToExecute]->Execute();
+						}
+					});
+			}
+		}
+	}
+
 	void AddStep(TUniquePtr<Private::IStep>&& Step)
 	{
 		Steps.Add(MoveTemp(Step));
@@ -1098,7 +1492,11 @@ public:
 			break;
 
 		case EOnlineAsyncExecutionPolicy::RunOnNextTick:
-			ExecuteOnGameThread(OpType::Name, MoveTemp(Callable));
+			FTSTicker::GetCoreTicker().AddTicker(FTickerDelegate::CreateLambda([LambdaCallable = Callable](float dt) mutable
+				{
+					ExecuteOnGameThread(OpType::Name, MoveTemp(LambdaCallable));
+					return false;
+				}));
 			break;
 
 		case EOnlineAsyncExecutionPolicy::RunOnThreadPool:
@@ -1123,16 +1521,31 @@ protected:
 
 	void SetResultAndState(TOnlineResult<OpType>&& Result, EAsyncOpState State)
 	{
-		if (SharedState->State <= EAsyncOpState::Queued)
+		if (IsInGameThread())
 		{
-			OnStartEvent.Broadcast(*this);
-		}
+			if (SharedState->State <= EAsyncOpState::Queued)
+			{
+				OnStartEvent.Broadcast(*this);
+			}
 
-		SharedState->Result = MoveTemp(Result);
-		SharedState->State = State;
-		UE_LOG(LogOnlineServices, Verbose, TEXT("%p %s op state set to: %s"), this, OpType::Name, LexToString(State));
+			SharedState->Result = MoveTemp(Result);
+			SharedState->State = State;
+			UE_LOG(LogOnlineServices, Verbose, TEXT("%p %s op state set to: %s"), this, OpType::Name, LexToString(State));
 
-		TriggerOnComplete(SharedState->Result);
+			TriggerOnComplete(SharedState->Result);
+		}
+		else
+		{
+			Execute(FOnlineAsyncExecutionPolicy::RunOnGameThread(),
+				[State, LambdaResult = MoveTemp(Result), WeakThis = TWeakPtr<TOnlineAsyncOp<OpType>>(this->AsShared())]() mutable
+				{
+					TSharedPtr<TOnlineAsyncOp<OpType>> PinnedThis = WeakThis.Pin();
+					if (PinnedThis)
+					{
+						PinnedThis->SetResultAndState(MoveTemp(LambdaResult), State);
+					}
+				});
+		}
 	}
 
 	void TriggerOnComplete(const TOnlineResult<OpType>& Result)
@@ -1304,9 +1717,9 @@ auto TOnlineAsyncOpBase<Outer, OpType, LastResultType>::Then(CallableType&& InCa
 
 	Op.AddStep(MoveTemp(StepPtr));
 
-	if constexpr (std::is_same_v<ResultType, void>)
+	if constexpr (std::is_same_v<TContinuationResultType_T<ResultType>, void>)
 	{
-		return TOnlineChainableAsyncOp<OpType, ResultType>(Op);
+		return TOnlineChainableAsyncOp<OpType, void>(Op);
 	}
 	else
 	{
@@ -1362,10 +1775,9 @@ auto TOnlineAsyncOpBase<Outer, OpType, void>::Then(CallableType&& InCallable, FO
 	Step->SetExecFunction(Op, MoveTemp(InCallable));
 
 	Op.AddStep(MoveTemp(StepPtr));
-
-	if constexpr (std::is_same_v<ResultType, void>)
+	if constexpr (std::is_same_v<TContinuationResultType_T<ResultType>, void>)
 	{
-		return TOnlineChainableAsyncOp<OpType, ResultType>(Op);
+		return TOnlineChainableAsyncOp<OpType, void>(Op);
 	}
 	else
 	{
