Make sure to wait for operations created by GetOp as well when shutdown OnlineServices instance.

diff --git a/Engine/Plugins/Online/OnlineServices/Source/OnlineServicesCommon/Private/Online/OnlineAsyncOpCache.cpp b/Engine/Plugins/Online/OnlineServices/Source/OnlineServicesCommon/Private/Online/OnlineAsyncOpCache.cpp
--- a/Engine/Plugins/Online/OnlineServices/Source/OnlineServicesCommon/Private/Online/OnlineAsyncOpCache.cpp
+++ b/Engine/Plugins/Online/OnlineServices/Source/OnlineServicesCommon/Private/Online/OnlineAsyncOpCache.cpp
@@ -13,6 +13,11 @@ TSharedRef<FOnlineAsyncOpCache> FOnlineAsyncOpCache::GetSharedThis()
 
 void FOnlineAsyncOpCache::ClearAllCallbacks()
 {
+	for (const TUniquePtr<IWrappedOperation>& WrappedOp : IndependentOperations)
+	{
+		WrappedOp->ClearCallback();
+	}
+
 	for (const TPair<FAccountId, TMap<FWrappedOperationKey, TUniquePtr<IWrappedOperation>>>& UserOperationPair : UserOperations)
 	{
 		ClearCallbacks(UserOperationPair.Value);
@@ -23,6 +28,13 @@ void FOnlineAsyncOpCache::ClearAllCallbacks()
 
 void FOnlineAsyncOpCache::CancelAll()
 {
+	// Move to temp array instead of using IndependentOperations directly because Cancel will remove from that array during for loop
+	TArray<TUniquePtr<IWrappedOperation>> TempIndependentOperations = MoveTemp(IndependentOperations);
+	for (const TUniquePtr<IWrappedOperation>& WrappedOp : TempIndependentOperations)
+	{
+		WrappedOp->Cancel();
+	}
+
 	for (const TPair<FAccountId, TMap<FWrappedOperationKey, TUniquePtr<IWrappedOperation>>>& UserOperationPair : UserOperations)
 	{
 		CancelOperations(UserOperationPair.Value);
@@ -33,6 +45,14 @@ void FOnlineAsyncOpCache::CancelAll()
 
 bool FOnlineAsyncOpCache::HasAnyRunningOperation() const
 {
+	for (const TUniquePtr<IWrappedOperation>& WrappedOp : IndependentOperations)
+	{
+		if (WrappedOp->GetAsyncOpState() < EAsyncOpState::Complete)
+		{
+			return true;
+		}
+	}
+
 	for (const TPair<FAccountId, TMap<FWrappedOperationKey, TUniquePtr<IWrappedOperation>>>& UserOperationPair : UserOperations)
 	{
 		if (HasAnyRunningOperationIn(UserOperationPair.Value))
diff --git a/Engine/Plugins/Online/OnlineServices/Source/OnlineServicesCommon/Public/Online/OnlineAsyncOpCache.h b/Engine/Plugins/Online/OnlineServices/Source/OnlineServicesCommon/Public/Online/OnlineAsyncOpCache.h
--- a/Engine/Plugins/Online/OnlineServices/Source/OnlineServicesCommon/Public/Online/OnlineAsyncOpCache.h
+++ b/Engine/Plugins/Online/OnlineServices/Source/OnlineServicesCommon/Public/Online/OnlineAsyncOpCache.h
@@ -248,7 +248,8 @@ public:
 	template <typename OpType>
 	TOnlineAsyncOpRef<OpType> GetOp(typename OpType::Params&& Params, const TArray<FString>& ConfigSectionHeiarchy)
 	{
-		TSharedRef<TOnlineAsyncOp<OpType>> Op = MakeShared<TOnlineAsyncOp<OpType>>(Services, MoveTemp(Params));
+		TUniquePtr<TWrappedOperation<OpType>> WrappedOp = MakeUnique<TWrappedOperation<OpType>>(Services, MoveTemp(Params));
+		TSharedRef<TOnlineAsyncOp<OpType>> Op = WrappedOp->GetDataRef();
 
 		if constexpr (TModels_V<CLocalAccountIdDefined, typename OpType::Params>)
 		{
@@ -256,6 +257,16 @@ public:
 			Op->Data.template Set<decltype(Params.LocalAccountId)>(TEXT("LocalAccountId"), Op->GetParams().LocalAccountId);
 		}
 
+		Op->OpCacheHandle = Op->OnComplete().Add(GetSharedThis(), [this](const TOnlineAsyncOp<OpType>& ThisOp, const TOnlineResult<OpType>& Result) {
+			const TOnlineAsyncOp<OpType>* ThisOpPtr = &ThisOp;
+			IndependentOperations.RemoveAll([ThisOpPtr](const TUniquePtr<IWrappedOperation>& WrappedOp) {
+				const TOnlineAsyncOp<OpType>* Op = WrappedOp->Get<TOnlineAsyncOp<OpType>>();
+				return ThisOpPtr == Op;
+			});
+		});
+
+		IndependentOperations.Add(MoveTemp(WrappedOp));
+
 		return Op;
 	}
 
@@ -545,6 +556,7 @@ private:
 
 	TMap<FWrappedOperationKey, TUniquePtr<IWrappedOperation>> Operations;
 	TMap<FAccountId, TMap<FWrappedOperationKey, TUniquePtr<IWrappedOperation>>> UserOperations;
+	TArray<TUniquePtr<IWrappedOperation>> IndependentOperations;
 };
 
 inline uint32 GetTypeHash(const FOnlineAsyncOpCache::FWrappedOperationKey& Key)
