Part 2
Changes to OSSV2 privileges interface.
Added option to control use of platform UI to resolve privilege failures.

diff --git a/Engine/Plugins/Online/OnlineServices/Source/OnlineServicesInterface/Private/Online/Privileges.cpp b/Engine/Plugins/Online/OnlineServices/Source/OnlineServicesInterface/Private/Online/Privileges.cpp
--- a/Engine/Plugins/Online/OnlineServices/Source/OnlineServicesInterface/Private/Online/Privileges.cpp
+++ b/Engine/Plugins/Online/OnlineServices/Source/OnlineServicesInterface/Private/Online/Privileges.cpp
@@ -139,4 +139,29 @@ void LexFromString(EPrivilegeResults& OutType, const FString& InStr)
 	}
 }
 
+FString LexToString(EShowPrivilegeResolutionUI Type)
+{
+	if (EnumHasAllFlags(Type, EShowPrivilegeResolutionUI::Default))		return TEXT("Default");
+	if (EnumHasAllFlags(Type, EShowPrivilegeResolutionUI::Show))			return TEXT("Show");
+	if (EnumHasAllFlags(Type, EShowPrivilegeResolutionUI::NotShow))		return TEXT("NotShow");
+	return TEXT("Invalid value");
+}
+
+void LexFromString(EShowPrivilegeResolutionUI& OutType, const FString& InStr)
+{
+	if (FCString::Stricmp(*InStr, TEXT("Show")) == 0)
+	{
+		OutType = EShowPrivilegeResolutionUI::Show;
+		return;
+	}
+	if (FCString::Stricmp(*InStr, TEXT("NotShow")) == 0)
+	{
+		OutType = EShowPrivilegeResolutionUI::NotShow;
+		return;
+	}	
+	//fallback to default
+	OutType = EShowPrivilegeResolutionUI::Default;
+	return;
+}
+
 /* UE::Online */}
\ No newline at end of file
diff --git a/Engine/Plugins/Online/OnlineServices/Source/OnlineServicesInterface/Public/Online/Privileges.h b/Engine/Plugins/Online/OnlineServices/Source/OnlineServicesInterface/Public/Online/Privileges.h
--- a/Engine/Plugins/Online/OnlineServices/Source/OnlineServicesInterface/Public/Online/Privileges.h
+++ b/Engine/Plugins/Online/OnlineServices/Source/OnlineServicesInterface/Public/Online/Privileges.h
@@ -28,6 +28,19 @@ enum class EUserPrivileges : uint8
 ONLINESERVICESINTERFACE_API const TCHAR* LexToString(EUserPrivileges Type);
 ONLINESERVICESINTERFACE_API void LexFromString(EUserPrivileges& OutType, const TCHAR* InStr);
 
+/** Whether to show UI to resolve the privilege if there is no access when GetUserPrivilege */
+enum class EShowPrivilegeResolutionUI : uint8
+{
+	/** Behavior matching the recommended practices for the platform */
+	Default,
+	/** Show UI to resolve privilege */
+	Show,
+	/** Not show UI to resolve privilege */
+	NotShow
+};
+ONLINESERVICESINTERFACE_API FString LexToString(EShowPrivilegeResolutionUI Type);
+ONLINESERVICESINTERFACE_API void LexFromString(EShowPrivilegeResolutionUI& OutType, const FString& InStr);
+
 enum class EPrivilegeResults : uint32
 {
 	/** The user has the requested privilege */
@@ -70,6 +83,9 @@ struct FQueryUserPrivilege
 		FAccountId LocalAccountId;
 		/** Privilege to query */
 		EUserPrivileges Privilege;
+		/** Should external UI be shown when not possessing the queried privilege */
+		EShowPrivilegeResolutionUI ShowExternalUI;
+
 	};
 
 	/** Output struct for Misc::FQueryUserPrivilege */
@@ -94,7 +110,8 @@ namespace Meta {
 
 BEGIN_ONLINE_STRUCT_META(FQueryUserPrivilege::Params)
 	ONLINE_STRUCT_FIELD(FQueryUserPrivilege::Params, LocalAccountId),
-	ONLINE_STRUCT_FIELD(FQueryUserPrivilege::Params, Privilege)
+	ONLINE_STRUCT_FIELD(FQueryUserPrivilege::Params, Privilege),
+	ONLINE_STRUCT_FIELD(FQueryUserPrivilege::Params, ShowExternalUI)
 END_ONLINE_STRUCT_META()
 
 BEGIN_ONLINE_STRUCT_META(FQueryUserPrivilege::Result)
