Include header file OnlineUtils.h properly in OnlineResult.h to fix the potential compile error when register commands in OSSv2, to find the proper overload function of UE::Online::ToLogString through metadata;
Remove the deprecated include OnlineUtilsCommon.h from OnlineUtils.h.

diff --git a/Engine/Plugins/Online/OnlineServices/Source/OnlineServicesInterface/Public/Online/OnlineResult.h b/Engine/Plugins/Online/OnlineServices/Source/OnlineServicesInterface/Public/Online/OnlineResult.h
--- a/Engine/Plugins/Online/OnlineServices/Source/OnlineServicesInterface/Public/Online/OnlineResult.h
+++ b/Engine/Plugins/Online/OnlineServices/Source/OnlineServicesInterface/Public/Online/OnlineResult.h
@@ -4,6 +4,7 @@
 
 #include "CoreGlobals.h"
 #include "Misc/TVariant.h"
+#include "Online/OnlineUtils.h"
 
 
 namespace UE::Online { class FOnlineError; }
diff --git a/Engine/Plugins/Online/OnlineServices/Source/OnlineServicesInterface/Public/Online/OnlineUtils.h b/Engine/Plugins/Online/OnlineServices/Source/OnlineServicesInterface/Public/Online/OnlineUtils.h
--- a/Engine/Plugins/Online/OnlineServices/Source/OnlineServicesInterface/Public/Online/OnlineUtils.h
+++ b/Engine/Plugins/Online/OnlineServices/Source/OnlineServicesInterface/Public/Online/OnlineUtils.h
@@ -8,9 +8,6 @@
 #include "Internationalization/Text.h"
 #include "Misc/TVariant.h"
 #include "Templates/SharedPointer.h"
-#if UE_ENABLE_INCLUDE_ORDER_DEPRECATED_IN_5_5
-#include "Online/OnlineUtilsCommon.h"
-#endif
 
 struct FDateTime;
 
