Created platform Privileges interface base and moving auth to a component

diff --git a/Engine/Plugins/Online/OnlineServices/Source/OnlineServicesCommon/Private/Online/PrivilegesCommon.cpp b/Engine/Plugins/Online/OnlineServices/Source/OnlineServicesCommon/Private/Online/PrivilegesCommon.cpp
--- a/Engine/Plugins/Online/OnlineServices/Source/OnlineServicesCommon/Private/Online/PrivilegesCommon.cpp
+++ b/Engine/Plugins/Online/OnlineServices/Source/OnlineServicesCommon/Private/Online/PrivilegesCommon.cpp
@@ -7,7 +7,7 @@
 namespace UE::Online {
 
 FPrivilegesCommon::FPrivilegesCommon(FOnlineServicesCommon& InServices)
-	: TOnlineComponent(TEXT("Permission"), InServices)
+	: TOnlineComponent(TEXT("Privileges"), InServices)
 {
 }
 
diff --git a/Engine/Plugins/Online/OnlineServices/Source/OnlineServicesCommon/Public/Online/PrivilegesCommon.h b/Engine/Plugins/Online/OnlineServices/Source/OnlineServicesCommon/Public/Online/PrivilegesCommon.h
--- a/Engine/Plugins/Online/OnlineServices/Source/OnlineServicesCommon/Public/Online/PrivilegesCommon.h
+++ b/Engine/Plugins/Online/OnlineServices/Source/OnlineServicesCommon/Public/Online/PrivilegesCommon.h
@@ -16,11 +16,13 @@ public:
 
 	FPrivilegesCommon(FOnlineServicesCommon& InServices);
 
-	// TOnlineComponent
+	// Begin TOnlineComponent
 	virtual void RegisterCommands() override;
+	// End TOnlineComponent
 
-	// IPermission
+	// Begin IPrivileges
 	virtual TOnlineAsyncOpHandle<FQueryUserPrivilege> QueryUserPrivilege(FQueryUserPrivilege::Params&& Params) override;
+	// End IPrivileges
 };
 
 /* UE::Online */}
\ No newline at end of file
diff --git a/Engine/Plugins/Online/OnlineServices/Source/OnlineServicesInterface/Public/Online/Privileges.h b/Engine/Plugins/Online/OnlineServices/Source/OnlineServicesInterface/Public/Online/Privileges.h
--- a/Engine/Plugins/Online/OnlineServices/Source/OnlineServicesInterface/Public/Online/Privileges.h
+++ b/Engine/Plugins/Online/OnlineServices/Source/OnlineServicesInterface/Public/Online/Privileges.h
@@ -100,9 +100,11 @@ class IPrivileges
 {
 public:
 	/**
-		* Queries about a given privilege for a user
-		*/
+	* Queries about a given privilege for a user
+	*/
 	virtual TOnlineAsyncOpHandle<FQueryUserPrivilege> QueryUserPrivilege(FQueryUserPrivilege::Params&& Params) = 0;
+
+	virtual ~IPrivileges() = default;
 };
 
 namespace Meta {
