Fix the issue the completed op doesn't get removed from IndependentOperations, the type wasn't compared properly;
Fix the issue the op could get stuck in OnlineServicesNull when flush, the op relies on CoreTicker;
Tick TaskGraph as well during FlushTick to make sure the game thread will process messages from those threads;
Set default OnlineAsyncOp flush waiting time to 0, only set it when needed.

diff --git a/Engine/Plugins/Online/OnlineServices/Source/OnlineServicesCommon/Private/Online/OnlineServicesCommon.cpp b/Engine/Plugins/Online/OnlineServices/Source/OnlineServicesCommon/Private/Online/OnlineServicesCommon.cpp
--- a/Engine/Plugins/Online/OnlineServices/Source/OnlineServicesCommon/Private/Online/OnlineServicesCommon.cpp
+++ b/Engine/Plugins/Online/OnlineServices/Source/OnlineServicesCommon/Private/Online/OnlineServicesCommon.cpp
@@ -12,7 +12,7 @@ struct FOnlineServicesCommonConfig
 {
 	int32 MaxConcurrentOperations = 16;
 	float SecondsToSleepForOutstandingOperations = 0.01;
-	float SecondsToFlushForOutstandingOperations = 60;
+	float SecondsToFlushForOutstandingOperations = 0;
 };
 
 namespace Meta {
@@ -95,6 +95,12 @@ void FOnlineServicesCommon::Flush(EAsyncOpFlushReason FlushReason)
 void FOnlineServicesCommon::FlushTick(float DeltaSeconds)
 {
 	Tick(DeltaSeconds);
+
+	// In case any AsyncOp relies on CoreTicker
+	FTSTicker::GetCoreTicker().Tick(DeltaSeconds);
+
+	// TaskGraph also needs to be ticked to process messages that come back to the game thread
+	FTaskGraphInterface::Get().ProcessThreadUntilIdle(ENamedThreads::GameThread);
 }
 
 void FOnlineServicesCommon::Destroy()
diff --git a/Engine/Plugins/Online/OnlineServices/Source/OnlineServicesCommon/Public/Online/OnlineAsyncOpCache.h b/Engine/Plugins/Online/OnlineServices/Source/OnlineServicesCommon/Public/Online/OnlineAsyncOpCache.h
--- a/Engine/Plugins/Online/OnlineServices/Source/OnlineServicesCommon/Public/Online/OnlineAsyncOpCache.h
+++ b/Engine/Plugins/Online/OnlineServices/Source/OnlineServicesCommon/Public/Online/OnlineAsyncOpCache.h
@@ -260,8 +260,12 @@ public:
 		Op->OpCacheHandle = Op->OnComplete().Add(GetSharedThis(), [this](const TOnlineAsyncOp<OpType>& ThisOp, const TOnlineResult<OpType>& Result) {
 			const TOnlineAsyncOp<OpType>* ThisOpPtr = &ThisOp;
 			IndependentOperations.RemoveAll([ThisOpPtr](const TUniquePtr<IWrappedOperation>& WrappedOp) {
-				const TOnlineAsyncOp<OpType>* Op = WrappedOp->Get<TOnlineAsyncOp<OpType>>();
-				return ThisOpPtr == Op;
+				const TSharedRef<TOnlineAsyncOp<OpType>>* Op = WrappedOp->Get<TSharedRef<TOnlineAsyncOp<OpType>>>();
+				if (Op)
+				{
+					return ThisOpPtr == &(*Op).Get();
+				}
+				return false;
 			});
 		});
 
