Add support of FUtf8String type in OnlineServices meta data.

diff --git a/Engine/Plugins/Online/OnlineServices/Source/OnlineServicesCommon/Public/Online/OnlineExecHandler.h b/Engine/Plugins/Online/OnlineServices/Source/OnlineServicesCommon/Public/Online/OnlineExecHandler.h
--- a/Engine/Plugins/Online/OnlineServices/Source/OnlineServicesCommon/Public/Online/OnlineExecHandler.h
+++ b/Engine/Plugins/Online/OnlineServices/Source/OnlineServicesCommon/Public/Online/OnlineExecHandler.h
@@ -283,6 +283,18 @@ inline bool ParseOnlineExecParams(const TCHAR*& Cmd, FName& Value, IOnlineServic
 	return true;
 }
 
+inline bool ParseOnlineExecParams(const TCHAR*& Cmd, FUtf8String& Value, IOnlineServices* Services = nullptr)
+{
+	FString StringValue;
+	if (!ParseOnlineExecParams(Cmd, StringValue))
+	{
+		return false;
+	}
+
+	Value = StringCast<UTF8CHAR>(*StringValue, StringValue.Len());
+	return true;
+}
+
 inline bool ParseOnlineExecParams(const TCHAR*& Cmd, uint8& Value, IOnlineServices* Services)
 {
 	FString Token;
diff --git a/Engine/Plugins/Online/OnlineServices/Source/OnlineServicesInterface/Public/Online/OnlineUtils.h b/Engine/Plugins/Online/OnlineServices/Source/OnlineServicesInterface/Public/Online/OnlineUtils.h
--- a/Engine/Plugins/Online/OnlineServices/Source/OnlineServicesInterface/Public/Online/OnlineUtils.h
+++ b/Engine/Plugins/Online/OnlineServices/Source/OnlineServicesInterface/Public/Online/OnlineUtils.h
@@ -5,6 +5,7 @@
 #include "Online/OnlineMeta.h"
 #include "Containers/Map.h"
 #include "Containers/UnrealString.h"
+#include "Containers/Utf8String.h"
 #include "Internationalization/Text.h"
 #include "Misc/TVariant.h"
 #include "Templates/SharedPointer.h"
@@ -24,6 +25,7 @@ template <typename... Ts> FString ToLogString(const TVariant<Ts...>& Variant);
 template <typename T, ESPMode Mode> FString ToLogString(const TSharedRef<T, Mode>& Ref);
 inline FString ToLogString(const FString& String);
 inline FString ToLogString(const FName& Name);
+inline FString ToLogString(const FUtf8String& Name);
 inline FString ToLogString(const FText& Text);
 inline FString ToLogString(uint8 Value);
 inline FString ToLogString(int8 Value);
@@ -115,6 +117,11 @@ inline FString ToLogString(const FName& Name)
 	return Name.ToString();
 }
 
+inline FString ToLogString(const FUtf8String& String)
+{
+	return FString(StringCast<TCHAR>(*String, String.Len()));
+}
+
 inline FString ToLogString(const FText& Text)
 {
 	return Text.ToString();
