Add support to specify different config value between Editor and Engine in OnlineServices;
Use colon instead of space for override section in online config, output warning when space detected in section of online config.

diff --git a/Engine/Plugins/Online/OnlineServices/Source/OnlineServicesCommon/Public/Online/OnlineConfig.h b/Engine/Plugins/Online/OnlineServices/Source/OnlineServicesCommon/Public/Online/OnlineConfig.h
--- a/Engine/Plugins/Online/OnlineServices/Source/OnlineServicesCommon/Public/Online/OnlineConfig.h
+++ b/Engine/Plugins/Online/OnlineServices/Source/OnlineServicesCommon/Public/Online/OnlineConfig.h
@@ -4,6 +4,7 @@
 
 #include "Internationalization/Text.h"
 #include "Online/OnlineMeta.h"
+#include "Online/OnlineServicesLog.h"
 
 namespace UE::Online { class FSchemaVariant; }
 
@@ -393,7 +394,12 @@ auto LoadConfig(IOnlineConfigProvider& Provider, const TArray<FString>& SectionH
 	bool bLoadedValue = false;
 	for (const FString& Section : SectionHeirarchy)
 	{
-		bLoadedValue |= LoadConfig(Provider, Section, Value);
+		bool bLoadedInCurrentSection = LoadConfig(Provider, Section, Value);
+		if (bLoadedInCurrentSection && Section.Contains(" "))
+		{
+			UE_LOG(LogOnlineServices, Warning, TEXT("Space in online config section has been deprecated, use : instead in [%s]"), *Section);
+		}
+		bLoadedValue |= bLoadedInCurrentSection;
 	}
 	return bLoadedValue;
 }
diff --git a/Engine/Plugins/Online/OnlineServices/Source/OnlineServicesCommon/Public/Online/OnlineServicesCommon.h b/Engine/Plugins/Online/OnlineServices/Source/OnlineServicesCommon/Public/Online/OnlineServicesCommon.h
--- a/Engine/Plugins/Online/OnlineServices/Source/OnlineServicesCommon/Public/Online/OnlineServicesCommon.h
+++ b/Engine/Plugins/Online/OnlineServices/Source/OnlineServicesCommon/Public/Online/OnlineServicesCommon.h
@@ -245,14 +245,26 @@ public:
 	 */
 	TArray<FString> GetConfigSectionHeirachWithOverrides(const TArray<FString>& SectionHeiarchy) const
 	{
+		TArray<FString> AllConfigSectionOverrides = ConfigSectionOverrides;
+
+		if (GIsEditor)
+		{
+			AllConfigSectionOverrides.Add("Editor");
+		}
+
 		TArray<FString> SectionHeiarchyWithOverrides;
 		for (const FString& Section : SectionHeiarchy)
 		{
 			SectionHeiarchyWithOverrides.Add(Section);
-			for (const FString& Override : ConfigSectionOverrides)
+
+			for (const FString& Override : AllConfigSectionOverrides)
 			{
+				// Using of space in online config section is deprecated
 				FString OverrideSection = Section + TEXT(" ") + Override;
 				SectionHeiarchyWithOverrides.Add(OverrideSection);
+
+				OverrideSection = Section + TEXT(":") + Override;
+				SectionHeiarchyWithOverrides.Add(OverrideSection);
 			}
 		}
 
