[OnlineServicesCommon]
* Moved dll storage to methods/variables instead of on types. This is needed to get merged modules working

diff --git a/Engine/Plugins/Online/OnlineServices/Source/OnlineServicesCommon/Public/Online/AchievementsCommon.h b/Engine/Plugins/Online/OnlineServices/Source/OnlineServicesCommon/Public/Online/AchievementsCommon.h
--- a/Engine/Plugins/Online/OnlineServices/Source/OnlineServicesCommon/Public/Online/AchievementsCommon.h
+++ b/Engine/Plugins/Online/OnlineServices/Source/OnlineServicesCommon/Public/Online/AchievementsCommon.h
@@ -6,6 +6,8 @@
 #include "Online/OnlineComponent.h"
 #include "Online/Stats.h"
 
+#define UE_API ONLINESERVICESCOMMON_API
+
 namespace UE::Online {
 
 class FOnlineServicesCommon;
@@ -50,38 +52,38 @@ END_ONLINE_STRUCT_META()
 
 /* Meta */ }
 
-class ONLINESERVICESCOMMON_API FAchievementsCommon : public TOnlineComponent<IAchievements>
+class FAchievementsCommon : public TOnlineComponent<IAchievements>
 {
 public:
 	using Super = IAchievements;
 
-	FAchievementsCommon(FOnlineServicesCommon& InServices);
+	UE_API FAchievementsCommon(FOnlineServicesCommon& InServices);
 
 	// TOnlineComponent
-	virtual void Initialize() override;
-	virtual void Shutdown() override;
-	virtual void UpdateConfig() override;
-	virtual void RegisterCommands() override;
+	UE_API virtual void Initialize() override;
+	UE_API virtual void Shutdown() override;
+	UE_API virtual void UpdateConfig() override;
+	UE_API virtual void RegisterCommands() override;
 
 	// IAchievements
-	virtual TOnlineAsyncOpHandle<FQueryAchievementDefinitions> QueryAchievementDefinitions(FQueryAchievementDefinitions::Params&& Params) override;
-	virtual TOnlineResult<FGetAchievementIds> GetAchievementIds(FGetAchievementIds::Params&& Params) override;
-	virtual TOnlineResult<FGetAchievementDefinition> GetAchievementDefinition(FGetAchievementDefinition::Params&& Params) override;
-	virtual TOnlineAsyncOpHandle<FQueryAchievementStates> QueryAchievementStates(FQueryAchievementStates::Params&& Params) override;
-	virtual TOnlineResult<FGetAchievementState> GetAchievementState(FGetAchievementState::Params&& Params) const override;
-	virtual TOnlineAsyncOpHandle<FUnlockAchievements> UnlockAchievements(FUnlockAchievements::Params&& Params) override;
-	virtual TOnlineResult<FDisplayAchievementUI> DisplayAchievementUI(FDisplayAchievementUI::Params&& Params) override;
-	virtual TOnlineEvent<void(const FAchievementStateUpdated&)> OnAchievementStateUpdated() override;
+	UE_API virtual TOnlineAsyncOpHandle<FQueryAchievementDefinitions> QueryAchievementDefinitions(FQueryAchievementDefinitions::Params&& Params) override;
+	UE_API virtual TOnlineResult<FGetAchievementIds> GetAchievementIds(FGetAchievementIds::Params&& Params) override;
+	UE_API virtual TOnlineResult<FGetAchievementDefinition> GetAchievementDefinition(FGetAchievementDefinition::Params&& Params) override;
+	UE_API virtual TOnlineAsyncOpHandle<FQueryAchievementStates> QueryAchievementStates(FQueryAchievementStates::Params&& Params) override;
+	UE_API virtual TOnlineResult<FGetAchievementState> GetAchievementState(FGetAchievementState::Params&& Params) const override;
+	UE_API virtual TOnlineAsyncOpHandle<FUnlockAchievements> UnlockAchievements(FUnlockAchievements::Params&& Params) override;
+	UE_API virtual TOnlineResult<FDisplayAchievementUI> DisplayAchievementUI(FDisplayAchievementUI::Params&& Params) override;
+	UE_API virtual TOnlineEvent<void(const FAchievementStateUpdated&)> OnAchievementStateUpdated() override;
 
 protected:
 	TOnlineEventCallable<void(const FAchievementStateUpdated&)> OnAchievementStateUpdatedEvent;
 
-	void OnAchievementStatesQueried(const FAccountId& AccountId);
+	UE_API void OnAchievementStatesQueried(const FAccountId& AccountId);
 
-	void UnlockAchievementsByStats(const FStatsUpdated& StatsUpdated);
-	void ExecuteUnlockRulesRelatedToStat(const FAccountId& AccountId, const FString& StatName, const TMap<FString, FStatValue>& Stats, TArray<FString>& OutAchievementsToUnlock);
-	bool MeetUnlockCondition(const FAchievementUnlockRule& AchievementUnlockRule, const TMap<FString, FStatValue>& Stats);
-	bool IsUnlocked(const FAccountId& AccountId, const FString& AchievementName) const;
+	UE_API void UnlockAchievementsByStats(const FStatsUpdated& StatsUpdated);
+	UE_API void ExecuteUnlockRulesRelatedToStat(const FAccountId& AccountId, const FString& StatName, const TMap<FString, FStatValue>& Stats, TArray<FString>& OutAchievementsToUnlock);
+	UE_API bool MeetUnlockCondition(const FAchievementUnlockRule& AchievementUnlockRule, const TMap<FString, FStatValue>& Stats);
+	UE_API bool IsUnlocked(const FAccountId& AccountId, const FString& AchievementName) const;
 
 	FOnlineEventDelegateHandle StatEventHandle;
 
@@ -92,3 +94,5 @@ protected:
 };
 
 /* UE::Online */ }
+
+#undef UE_API
diff --git a/Engine/Plugins/Online/OnlineServices/Source/OnlineServicesCommon/Public/Online/AuthCommon.h b/Engine/Plugins/Online/OnlineServices/Source/OnlineServicesCommon/Public/Online/AuthCommon.h
--- a/Engine/Plugins/Online/OnlineServices/Source/OnlineServicesCommon/Public/Online/AuthCommon.h
+++ b/Engine/Plugins/Online/OnlineServices/Source/OnlineServicesCommon/Public/Online/AuthCommon.h
@@ -4,26 +4,28 @@
 
 #include "Online/OnlineComponent.h"
 
+#define UE_API ONLINESERVICESCOMMON_API
+
 namespace UE::Online {
 
 class FOnlineServicesCommon;
 
-class ONLINESERVICESCOMMON_API FAccountInfoRegistry
+class FAccountInfoRegistry
 {
 public:
 	virtual ~FAccountInfoRegistry() = default;
 
-	TSharedPtr<FAccountInfo> Find(FPlatformUserId PlatformUserId) const;
-	TSharedPtr<FAccountInfo> Find(FAccountId AccountId) const;
+	UE_API TSharedPtr<FAccountInfo> Find(FPlatformUserId PlatformUserId) const;
+	UE_API TSharedPtr<FAccountInfo> Find(FAccountId AccountId) const;
 
-	TArray<TSharedRef<FAccountInfo>> GetAllAccountInfo(TFunction<bool(const TSharedRef<FAccountInfo>&)> Predicate) const;
+	UE_API TArray<TSharedRef<FAccountInfo>> GetAllAccountInfo(TFunction<bool(const TSharedRef<FAccountInfo>&)> Predicate) const;
 
 protected:
-	TSharedPtr<FAccountInfo> FindNoLock(FPlatformUserId PlatformUserId) const;
-	TSharedPtr<FAccountInfo> FindNoLock(FAccountId AccountId) const;
+	UE_API TSharedPtr<FAccountInfo> FindNoLock(FPlatformUserId PlatformUserId) const;
+	UE_API TSharedPtr<FAccountInfo> FindNoLock(FAccountId AccountId) const;
 
-	virtual void DoRegister(const TSharedRef<FAccountInfo>& AccountInfo);
-	virtual void DoUnregister(const TSharedRef<FAccountInfo>& AccountInfo);
+	UE_API virtual void DoRegister(const TSharedRef<FAccountInfo>& AccountInfo);
+	UE_API virtual void DoUnregister(const TSharedRef<FAccountInfo>& AccountInfo);
 
 	mutable FRWLock IndexLock;
 
@@ -32,38 +34,38 @@ private:
 	TMap<FAccountId, TSharedRef<FAccountInfo>> AuthDataByOnlineAccountIdHandle;
 };
 
-class ONLINESERVICESCOMMON_API FAuthCommon : public TOnlineComponent<IAuth>
+class FAuthCommon : public TOnlineComponent<IAuth>
 {
 public:
 	using Super = IAuth;
 
-	FAuthCommon(FOnlineServicesCommon& InServices);
+	UE_API FAuthCommon(FOnlineServicesCommon& InServices);
 
 	// IOnlineComponent
-	virtual void RegisterCommands() override;
+	UE_API virtual void RegisterCommands() override;
 
 	// IAuth
-	virtual TOnlineAsyncOpHandle<FAuthLogin> Login(FAuthLogin::Params&& Params) override;
-	virtual TOnlineAsyncOpHandle<FAuthLogout> Logout(FAuthLogout::Params&& Params) override;
-	virtual TOnlineAsyncOpHandle<FAuthCreateAccount> CreateAccount(FAuthCreateAccount::Params&& Params) override;
-	virtual TOnlineAsyncOpHandle<FAuthLinkAccount> LinkAccount(FAuthLinkAccount::Params&& Params) override;
-	virtual TOnlineAsyncOpHandle<FAuthModifyAccountAttributes> ModifyAccountAttributes(FAuthModifyAccountAttributes::Params&& Params) override;
-	virtual TOnlineAsyncOpHandle<FAuthQueryExternalServerAuthTicket> QueryExternalServerAuthTicket(FAuthQueryExternalServerAuthTicket::Params&& Params) override;
-	virtual TOnlineAsyncOpHandle<FAuthQueryExternalAuthToken> QueryExternalAuthToken(FAuthQueryExternalAuthToken::Params&& Params) override;
-	virtual TOnlineAsyncOpHandle<FAuthQueryVerifiedAuthTicket> QueryVerifiedAuthTicket(FAuthQueryVerifiedAuthTicket::Params&& Params) override;
-	virtual TOnlineAsyncOpHandle<FAuthCancelVerifiedAuthTicket> CancelVerifiedAuthTicket(FAuthCancelVerifiedAuthTicket::Params&& Params) override;
-	virtual TOnlineAsyncOpHandle<FAuthBeginVerifiedAuthSession> BeginVerifiedAuthSession(FAuthBeginVerifiedAuthSession::Params&& Params) override;
-	virtual TOnlineAsyncOpHandle<FAuthEndVerifiedAuthSession> EndVerifiedAuthSession(FAuthEndVerifiedAuthSession::Params&& Params) override;
-	virtual TOnlineResult<FAuthGetLocalOnlineUserByOnlineAccountId> GetLocalOnlineUserByOnlineAccountId(FAuthGetLocalOnlineUserByOnlineAccountId::Params&& Params) const override;
-	virtual TOnlineResult<FAuthGetLocalOnlineUserByPlatformUserId> GetLocalOnlineUserByPlatformUserId(FAuthGetLocalOnlineUserByPlatformUserId::Params&& Params) const override;
-	virtual TOnlineResult<FAuthGetAllLocalOnlineUsers> GetAllLocalOnlineUsers(FAuthGetAllLocalOnlineUsers::Params&& Params) const override;
-	virtual TOnlineResult<FAuthGetLinkAccountContinuationId> GetLinkAccountContinuationId(FAuthGetLinkAccountContinuationId::Params&& Params) const override;
-	virtual TOnlineResult<FAuthGetRelyingParty> GetRelyingParty(FAuthGetRelyingParty::Params&& Params) const override;
-	virtual TOnlineEvent<void(const FAuthLoginStatusChanged&)> OnLoginStatusChanged() override;
-	virtual TOnlineEvent<void(const FAuthPendingAuthExpiration&)> OnPendingAuthExpiration() override;
-	virtual TOnlineEvent<void(const FAuthAccountAttributesChanged&)> OnAccountAttributesChanged() override;
-	virtual bool IsLoggedIn(const FAccountId& AccountId) const override;
-	virtual bool IsLoggedIn(const FPlatformUserId& PlatformUserId) const override;
+	UE_API virtual TOnlineAsyncOpHandle<FAuthLogin> Login(FAuthLogin::Params&& Params) override;
+	UE_API virtual TOnlineAsyncOpHandle<FAuthLogout> Logout(FAuthLogout::Params&& Params) override;
+	UE_API virtual TOnlineAsyncOpHandle<FAuthCreateAccount> CreateAccount(FAuthCreateAccount::Params&& Params) override;
+	UE_API virtual TOnlineAsyncOpHandle<FAuthLinkAccount> LinkAccount(FAuthLinkAccount::Params&& Params) override;
+	UE_API virtual TOnlineAsyncOpHandle<FAuthModifyAccountAttributes> ModifyAccountAttributes(FAuthModifyAccountAttributes::Params&& Params) override;
+	UE_API virtual TOnlineAsyncOpHandle<FAuthQueryExternalServerAuthTicket> QueryExternalServerAuthTicket(FAuthQueryExternalServerAuthTicket::Params&& Params) override;
+	UE_API virtual TOnlineAsyncOpHandle<FAuthQueryExternalAuthToken> QueryExternalAuthToken(FAuthQueryExternalAuthToken::Params&& Params) override;
+	UE_API virtual TOnlineAsyncOpHandle<FAuthQueryVerifiedAuthTicket> QueryVerifiedAuthTicket(FAuthQueryVerifiedAuthTicket::Params&& Params) override;
+	UE_API virtual TOnlineAsyncOpHandle<FAuthCancelVerifiedAuthTicket> CancelVerifiedAuthTicket(FAuthCancelVerifiedAuthTicket::Params&& Params) override;
+	UE_API virtual TOnlineAsyncOpHandle<FAuthBeginVerifiedAuthSession> BeginVerifiedAuthSession(FAuthBeginVerifiedAuthSession::Params&& Params) override;
+	UE_API virtual TOnlineAsyncOpHandle<FAuthEndVerifiedAuthSession> EndVerifiedAuthSession(FAuthEndVerifiedAuthSession::Params&& Params) override;
+	UE_API virtual TOnlineResult<FAuthGetLocalOnlineUserByOnlineAccountId> GetLocalOnlineUserByOnlineAccountId(FAuthGetLocalOnlineUserByOnlineAccountId::Params&& Params) const override;
+	UE_API virtual TOnlineResult<FAuthGetLocalOnlineUserByPlatformUserId> GetLocalOnlineUserByPlatformUserId(FAuthGetLocalOnlineUserByPlatformUserId::Params&& Params) const override;
+	UE_API virtual TOnlineResult<FAuthGetAllLocalOnlineUsers> GetAllLocalOnlineUsers(FAuthGetAllLocalOnlineUsers::Params&& Params) const override;
+	UE_API virtual TOnlineResult<FAuthGetLinkAccountContinuationId> GetLinkAccountContinuationId(FAuthGetLinkAccountContinuationId::Params&& Params) const override;
+	UE_API virtual TOnlineResult<FAuthGetRelyingParty> GetRelyingParty(FAuthGetRelyingParty::Params&& Params) const override;
+	UE_API virtual TOnlineEvent<void(const FAuthLoginStatusChanged&)> OnLoginStatusChanged() override;
+	UE_API virtual TOnlineEvent<void(const FAuthPendingAuthExpiration&)> OnPendingAuthExpiration() override;
+	UE_API virtual TOnlineEvent<void(const FAuthAccountAttributesChanged&)> OnAccountAttributesChanged() override;
+	UE_API virtual bool IsLoggedIn(const FAccountId& AccountId) const override;
+	UE_API virtual bool IsLoggedIn(const FPlatformUserId& PlatformUserId) const override;
 
 protected:
 	virtual const FAccountInfoRegistry& GetAccountInfoRegistry() const = 0;
@@ -74,3 +76,5 @@ protected:
 };
 
 /* UE::Online */ }
+
+#undef UE_API
diff --git a/Engine/Plugins/Online/OnlineServices/Source/OnlineServicesCommon/Public/Online/CommerceCommon.h b/Engine/Plugins/Online/OnlineServices/Source/OnlineServicesCommon/Public/Online/CommerceCommon.h
--- a/Engine/Plugins/Online/OnlineServices/Source/OnlineServicesCommon/Public/Online/CommerceCommon.h
+++ b/Engine/Plugins/Online/OnlineServices/Source/OnlineServicesCommon/Public/Online/CommerceCommon.h
@@ -5,36 +5,38 @@
 #include "Online/Commerce.h"
 #include "Online/OnlineComponent.h"
 
+#define UE_API ONLINESERVICESCOMMON_API
+
 namespace UE::Online {
 
 class FOnlineServicesCommon;
 
-class ONLINESERVICESCOMMON_API FCommerceCommon : public TOnlineComponent<ICommerce>
+class FCommerceCommon : public TOnlineComponent<ICommerce>
 {
 public:
 	using Super = ICommerce;
 
-	FCommerceCommon(FOnlineServicesCommon& InServices);
+	UE_API FCommerceCommon(FOnlineServicesCommon& InServices);
 
 	// IOnlineComponent
-	virtual void RegisterCommands() override;
+	UE_API virtual void RegisterCommands() override;
 
 	// ICommerce
-	virtual TOnlineAsyncOpHandle<FCommerceQueryOffers> QueryOffers(FCommerceQueryOffers::Params&& Params) override;
-	virtual TOnlineAsyncOpHandle<FCommerceQueryOffersById> QueryOffersById(FCommerceQueryOffersById::Params&& Params) override;
-	virtual TOnlineResult<FCommerceGetOffers> GetOffers(FCommerceGetOffers::Params&& Params) override;
-	virtual TOnlineResult<FCommerceGetOffersById> GetOffersById(FCommerceGetOffersById::Params&& Params) override;
-	virtual TOnlineAsyncOpHandle<FCommerceShowStoreUI> ShowStoreUI(FCommerceShowStoreUI::Params&& Params) override;
-	virtual TOnlineAsyncOpHandle<FCommerceCheckout> Checkout(FCommerceCheckout::Params&& Params) override;
-	virtual TOnlineAsyncOpHandle<FCommerceQueryTransactionEntitlements> QueryTransactionEntitlements(FCommerceQueryTransactionEntitlements::Params&& Params) override;
-	virtual TOnlineAsyncOpHandle<FCommerceQueryEntitlements> QueryEntitlements(FCommerceQueryEntitlements::Params&& Params) override;
-	virtual TOnlineResult<FCommerceGetEntitlements> GetEntitlements(FCommerceGetEntitlements::Params&& Params) override;
-	virtual TOnlineAsyncOpHandle<FCommerceRedeemEntitlement> RedeemEntitlement(FCommerceRedeemEntitlement::Params&& Params) override;
-	virtual TOnlineAsyncOpHandle<FCommerceRetrieveS2SToken> RetrieveS2SToken(FCommerceRetrieveS2SToken::Params&& Params) override;
-	virtual TOnlineEvent<void(const FCommerceOnPurchaseComplete&)> OnPurchaseCompleted() override;
+	UE_API virtual TOnlineAsyncOpHandle<FCommerceQueryOffers> QueryOffers(FCommerceQueryOffers::Params&& Params) override;
+	UE_API virtual TOnlineAsyncOpHandle<FCommerceQueryOffersById> QueryOffersById(FCommerceQueryOffersById::Params&& Params) override;
+	UE_API virtual TOnlineResult<FCommerceGetOffers> GetOffers(FCommerceGetOffers::Params&& Params) override;
+	UE_API virtual TOnlineResult<FCommerceGetOffersById> GetOffersById(FCommerceGetOffersById::Params&& Params) override;
+	UE_API virtual TOnlineAsyncOpHandle<FCommerceShowStoreUI> ShowStoreUI(FCommerceShowStoreUI::Params&& Params) override;
+	UE_API virtual TOnlineAsyncOpHandle<FCommerceCheckout> Checkout(FCommerceCheckout::Params&& Params) override;
+	UE_API virtual TOnlineAsyncOpHandle<FCommerceQueryTransactionEntitlements> QueryTransactionEntitlements(FCommerceQueryTransactionEntitlements::Params&& Params) override;
+	UE_API virtual TOnlineAsyncOpHandle<FCommerceQueryEntitlements> QueryEntitlements(FCommerceQueryEntitlements::Params&& Params) override;
+	UE_API virtual TOnlineResult<FCommerceGetEntitlements> GetEntitlements(FCommerceGetEntitlements::Params&& Params) override;
+	UE_API virtual TOnlineAsyncOpHandle<FCommerceRedeemEntitlement> RedeemEntitlement(FCommerceRedeemEntitlement::Params&& Params) override;
+	UE_API virtual TOnlineAsyncOpHandle<FCommerceRetrieveS2SToken> RetrieveS2SToken(FCommerceRetrieveS2SToken::Params&& Params) override;
+	UE_API virtual TOnlineEvent<void(const FCommerceOnPurchaseComplete&)> OnPurchaseCompleted() override;
 
 	//CommerceCommon
-	FText GetFormattedPrice(uint64 Price, int32 DecimalPoint, FString CurrencyCode);
+	UE_API FText GetFormattedPrice(uint64 Price, int32 DecimalPoint, FString CurrencyCode);
 
 protected:
 	TOnlineEventCallable<void(const FCommerceOnPurchaseComplete&)> OnPurchaseCompletedEvent;
@@ -42,3 +44,5 @@ protected:
 
 /* UE::Online */
 }
+
+#undef UE_API
diff --git a/Engine/Plugins/Online/OnlineServices/Source/OnlineServicesCommon/Public/Online/ConnectivityCommon.h b/Engine/Plugins/Online/OnlineServices/Source/OnlineServicesCommon/Public/Online/ConnectivityCommon.h
--- a/Engine/Plugins/Online/OnlineServices/Source/OnlineServicesCommon/Public/Online/ConnectivityCommon.h
+++ b/Engine/Plugins/Online/OnlineServices/Source/OnlineServicesCommon/Public/Online/ConnectivityCommon.h
@@ -5,23 +5,25 @@
 #include "Online/Connectivity.h"
 #include "Online/OnlineComponent.h"
 
+#define UE_API ONLINESERVICESCOMMON_API
+
 namespace UE::Online {
 
 class FOnlineServicesCommon;
 
-class ONLINESERVICESCOMMON_API FConnectivityCommon : public TOnlineComponent<IConnectivity>
+class FConnectivityCommon : public TOnlineComponent<IConnectivity>
 {
 public:
 	using Super = IConnectivity;
 
-	FConnectivityCommon(FOnlineServicesCommon& InServices);
+	UE_API FConnectivityCommon(FOnlineServicesCommon& InServices);
 
 	// TOnlineComponent
-	virtual void RegisterCommands() override;
+	UE_API virtual void RegisterCommands() override;
 
 	// IConnectivity
-	virtual TOnlineResult<FGetConnectionStatus> GetConnectionStatus(FGetConnectionStatus::Params&& Params) override;
-	virtual TOnlineEvent<void(const FConnectionStatusChanged&)> OnConnectionStatusChanged() override;
+	UE_API virtual TOnlineResult<FGetConnectionStatus> GetConnectionStatus(FGetConnectionStatus::Params&& Params) override;
+	UE_API virtual TOnlineEvent<void(const FConnectionStatusChanged&)> OnConnectionStatusChanged() override;
 
 protected:
 	EOnlineServicesConnectionStatus ConnectionStatus = EOnlineServicesConnectionStatus::NotConnected;
@@ -29,4 +31,6 @@ protected:
 	TOnlineEventCallable<void(const FConnectionStatusChanged&)> OnConnectionStatusChangedEvent;
 };
 
-/* UE::Online */}
\ No newline at end of file
+/* UE::Online */}
+
+#undef UE_API
diff --git a/Engine/Plugins/Online/OnlineServices/Source/OnlineServicesCommon/Public/Online/ExternalUICommon.h b/Engine/Plugins/Online/OnlineServices/Source/OnlineServicesCommon/Public/Online/ExternalUICommon.h
--- a/Engine/Plugins/Online/OnlineServices/Source/OnlineServicesCommon/Public/Online/ExternalUICommon.h
+++ b/Engine/Plugins/Online/OnlineServices/Source/OnlineServicesCommon/Public/Online/ExternalUICommon.h
@@ -5,27 +5,31 @@
 #include "Online/ExternalUI.h"
 #include "Online/OnlineComponent.h"
 
+#define UE_API ONLINESERVICESCOMMON_API
+
 namespace UE::Online {
 
 class FOnlineServicesCommon;
 
-class ONLINESERVICESCOMMON_API FExternalUICommon : public TOnlineComponent<IExternalUI>
+class FExternalUICommon : public TOnlineComponent<IExternalUI>
 {
 public:
 	using Super = IExternalUI;
 
-	FExternalUICommon(FOnlineServicesCommon& InServices);
+	UE_API FExternalUICommon(FOnlineServicesCommon& InServices);
 
 	// TOnlineComponent
-	virtual void RegisterCommands() override;
+	UE_API virtual void RegisterCommands() override;
 
 	// IExternalUI
-	virtual TOnlineAsyncOpHandle<FExternalUIShowLoginUI> ShowLoginUI(FExternalUIShowLoginUI::Params&& Params) override;
-	virtual TOnlineAsyncOpHandle<FExternalUIShowFriendsUI> ShowFriendsUI(FExternalUIShowFriendsUI::Params&& Params) override;
+	UE_API virtual TOnlineAsyncOpHandle<FExternalUIShowLoginUI> ShowLoginUI(FExternalUIShowLoginUI::Params&& Params) override;
+	UE_API virtual TOnlineAsyncOpHandle<FExternalUIShowFriendsUI> ShowFriendsUI(FExternalUIShowFriendsUI::Params&& Params) override;
 
-	virtual TOnlineEvent<void(const FExternalUIStatusChanged&)> OnExternalUIStatusChanged() override;
+	UE_API virtual TOnlineEvent<void(const FExternalUIStatusChanged&)> OnExternalUIStatusChanged() override;
 protected:
 	TOnlineEventCallable<void(const FExternalUIStatusChanged&)> OnExternalUIStatusChangedEvent;
 };
 
-/* UE::Online */ }
\ No newline at end of file
+/* UE::Online */ }
+
+#undef UE_API
diff --git a/Engine/Plugins/Online/OnlineServices/Source/OnlineServicesCommon/Public/Online/LeaderboardsCommon.h b/Engine/Plugins/Online/OnlineServices/Source/OnlineServicesCommon/Public/Online/LeaderboardsCommon.h
--- a/Engine/Plugins/Online/OnlineServices/Source/OnlineServicesCommon/Public/Online/LeaderboardsCommon.h
+++ b/Engine/Plugins/Online/OnlineServices/Source/OnlineServicesCommon/Public/Online/LeaderboardsCommon.h
@@ -6,6 +6,8 @@
 #include "Online/OnlineComponent.h"
 #include "Online/Stats.h"
 
+#define UE_API ONLINESERVICESCOMMON_API
+
 namespace UE::Online {
 
 enum class ELeaderboardUpdateMethod : uint8
@@ -79,28 +81,28 @@ struct FWriteLeaderboardScores
 	};
 };
 
-class ONLINESERVICESCOMMON_API FLeaderboardsCommon : public TOnlineComponent<ILeaderboards>
+class FLeaderboardsCommon : public TOnlineComponent<ILeaderboards>
 {
 public:
 	using Super = TOnlineComponent<ILeaderboards>;
 
-	FLeaderboardsCommon(FOnlineServicesCommon& InServices);
+	UE_API FLeaderboardsCommon(FOnlineServicesCommon& InServices);
 
 	// TOnlineComponent
-	virtual void Initialize() override;
-	virtual void UpdateConfig() override;
-	virtual void RegisterCommands() override;
+	UE_API virtual void Initialize() override;
+	UE_API virtual void UpdateConfig() override;
+	UE_API virtual void RegisterCommands() override;
 
 	// ILeaderboards
-	virtual TOnlineAsyncOpHandle<FReadEntriesForUsers> ReadEntriesForUsers(FReadEntriesForUsers::Params&& Params) override;
-	virtual TOnlineAsyncOpHandle<FReadEntriesAroundRank> ReadEntriesAroundRank(FReadEntriesAroundRank::Params&& Params) override;
-	virtual TOnlineAsyncOpHandle<FReadEntriesAroundUser> ReadEntriesAroundUser(FReadEntriesAroundUser::Params&& Params) override;
+	UE_API virtual TOnlineAsyncOpHandle<FReadEntriesForUsers> ReadEntriesForUsers(FReadEntriesForUsers::Params&& Params) override;
+	UE_API virtual TOnlineAsyncOpHandle<FReadEntriesAroundRank> ReadEntriesAroundRank(FReadEntriesAroundRank::Params&& Params) override;
+	UE_API virtual TOnlineAsyncOpHandle<FReadEntriesAroundUser> ReadEntriesAroundUser(FReadEntriesAroundUser::Params&& Params) override;
 
 	// Internal use only, for stats common implementation
-	virtual TOnlineAsyncOpHandle<FWriteLeaderboardScores> WriteLeaderboardScores(FWriteLeaderboardScores::Params&& Params);
+	UE_API virtual TOnlineAsyncOpHandle<FWriteLeaderboardScores> WriteLeaderboardScores(FWriteLeaderboardScores::Params&& Params);
 
 protected:
-	void WriteLeaderboardsByStats(const FStatsUpdated& StatsUpdated);
+	UE_API void WriteLeaderboardsByStats(const FStatsUpdated& StatsUpdated);
 
 	TMap<FString, FLeaderboardDefinition> LeaderboardDefinitions;
 	FOnlineEventDelegateHandle StatEventHandle;
@@ -121,3 +123,5 @@ END_ONLINE_STRUCT_META()
 /* Meta*/ }
 
 /* UE::Online */ }
+
+#undef UE_API
diff --git a/Engine/Plugins/Online/OnlineServices/Source/OnlineServicesCommon/Public/Online/LobbiesCommon.h b/Engine/Plugins/Online/OnlineServices/Source/OnlineServicesCommon/Public/Online/LobbiesCommon.h
--- a/Engine/Plugins/Online/OnlineServices/Source/OnlineServicesCommon/Public/Online/LobbiesCommon.h
+++ b/Engine/Plugins/Online/OnlineServices/Source/OnlineServicesCommon/Public/Online/LobbiesCommon.h
@@ -5,55 +5,57 @@
 #include "Online/OnlineComponent.h"
 #include "Online/LobbiesCommonTypes.h"
 
+#define UE_API ONLINESERVICESCOMMON_API
+
 namespace UE::Online {
 
 class FOnlineServicesCommon;
 
-class ONLINESERVICESCOMMON_API FLobbiesCommon : public TOnlineComponent<ILobbies>
+class FLobbiesCommon : public TOnlineComponent<ILobbies>
 {
 public:
 	using Super = ILobbies;
 
-	FLobbiesCommon(FOnlineServicesCommon& InServices);
+	UE_API FLobbiesCommon(FOnlineServicesCommon& InServices);
 
 	// TOnlineComponent
-	virtual void Initialize() override;
-	virtual void RegisterCommands() override;
+	UE_API virtual void Initialize() override;
+	UE_API virtual void RegisterCommands() override;
 
 	// ILobbies
-	virtual TOnlineAsyncOpHandle<FCreateLobby> CreateLobby(FCreateLobby::Params&& Params) override;
-	virtual TOnlineAsyncOpHandle<FFindLobbies> FindLobbies(FFindLobbies::Params&& Params) override;
-	virtual TOnlineAsyncOpHandle<FRestoreLobbies> RestoreLobbies(FRestoreLobbies::Params&& Params) override;
-	virtual TOnlineAsyncOpHandle<FJoinLobby> JoinLobby(FJoinLobby::Params&& Params) override;
-	virtual TOnlineAsyncOpHandle<FLeaveLobby> LeaveLobby(FLeaveLobby::Params&& Params) override;
-	virtual TOnlineAsyncOpHandle<FInviteLobbyMember> InviteLobbyMember(FInviteLobbyMember::Params&& Params) override;
-	virtual TOnlineAsyncOpHandle<FDeclineLobbyInvitation> DeclineLobbyInvitation(FDeclineLobbyInvitation::Params&& Params) override;
-	virtual TOnlineAsyncOpHandle<FKickLobbyMember> KickLobbyMember(FKickLobbyMember::Params&& Params) override;
-	virtual TOnlineAsyncOpHandle<FPromoteLobbyMember> PromoteLobbyMember(FPromoteLobbyMember::Params&& Params) override;
-	virtual TOnlineAsyncOpHandle<FModifyLobbySchema> ModifyLobbySchema(FModifyLobbySchema::Params&& Params) override;
-	virtual TOnlineAsyncOpHandle<FModifyLobbyJoinPolicy> ModifyLobbyJoinPolicy(FModifyLobbyJoinPolicy::Params&& Params) override;
-	virtual TOnlineAsyncOpHandle<FModifyLobbyAttributes> ModifyLobbyAttributes(FModifyLobbyAttributes::Params&& Params) override;
-	virtual TOnlineAsyncOpHandle<FModifyLobbyMemberAttributes> ModifyLobbyMemberAttributes(FModifyLobbyMemberAttributes::Params&& Params) override;
-	virtual TOnlineResult<FGetPresenceLobby> GetPresenceLobby(FGetPresenceLobby::Params&& Params) override;
-	virtual TOnlineResult<FIsPresenceLobby> IsPresenceLobby(FIsPresenceLobby::Params&& Params) override;
-	virtual TOnlineResult<FGetJoinedLobbies> GetJoinedLobbies(FGetJoinedLobbies::Params&& Params) override;
-	virtual TOnlineResult<FGetReceivedInvitations> GetReceivedInvitations(FGetReceivedInvitations::Params&& Params) override;
+	UE_API virtual TOnlineAsyncOpHandle<FCreateLobby> CreateLobby(FCreateLobby::Params&& Params) override;
+	UE_API virtual TOnlineAsyncOpHandle<FFindLobbies> FindLobbies(FFindLobbies::Params&& Params) override;
+	UE_API virtual TOnlineAsyncOpHandle<FRestoreLobbies> RestoreLobbies(FRestoreLobbies::Params&& Params) override;
+	UE_API virtual TOnlineAsyncOpHandle<FJoinLobby> JoinLobby(FJoinLobby::Params&& Params) override;
+	UE_API virtual TOnlineAsyncOpHandle<FLeaveLobby> LeaveLobby(FLeaveLobby::Params&& Params) override;
+	UE_API virtual TOnlineAsyncOpHandle<FInviteLobbyMember> InviteLobbyMember(FInviteLobbyMember::Params&& Params) override;
+	UE_API virtual TOnlineAsyncOpHandle<FDeclineLobbyInvitation> DeclineLobbyInvitation(FDeclineLobbyInvitation::Params&& Params) override;
+	UE_API virtual TOnlineAsyncOpHandle<FKickLobbyMember> KickLobbyMember(FKickLobbyMember::Params&& Params) override;
+	UE_API virtual TOnlineAsyncOpHandle<FPromoteLobbyMember> PromoteLobbyMember(FPromoteLobbyMember::Params&& Params) override;
+	UE_API virtual TOnlineAsyncOpHandle<FModifyLobbySchema> ModifyLobbySchema(FModifyLobbySchema::Params&& Params) override;
+	UE_API virtual TOnlineAsyncOpHandle<FModifyLobbyJoinPolicy> ModifyLobbyJoinPolicy(FModifyLobbyJoinPolicy::Params&& Params) override;
+	UE_API virtual TOnlineAsyncOpHandle<FModifyLobbyAttributes> ModifyLobbyAttributes(FModifyLobbyAttributes::Params&& Params) override;
+	UE_API virtual TOnlineAsyncOpHandle<FModifyLobbyMemberAttributes> ModifyLobbyMemberAttributes(FModifyLobbyMemberAttributes::Params&& Params) override;
+	UE_API virtual TOnlineResult<FGetPresenceLobby> GetPresenceLobby(FGetPresenceLobby::Params&& Params) override;
+	UE_API virtual TOnlineResult<FIsPresenceLobby> IsPresenceLobby(FIsPresenceLobby::Params&& Params) override;
+	UE_API virtual TOnlineResult<FGetJoinedLobbies> GetJoinedLobbies(FGetJoinedLobbies::Params&& Params) override;
+	UE_API virtual TOnlineResult<FGetReceivedInvitations> GetReceivedInvitations(FGetReceivedInvitations::Params&& Params) override;
 
-	virtual TOnlineEvent<void(const FLobbyJoined&)> OnLobbyJoined() override;
-	virtual TOnlineEvent<void(const FLobbyLeft&)> OnLobbyLeft() override;
-	virtual TOnlineEvent<void(const FLobbyMemberJoined&)> OnLobbyMemberJoined() override;
-	virtual TOnlineEvent<void(const FLobbyMemberLeft&)> OnLobbyMemberLeft() override;
-	virtual TOnlineEvent<void(const FLobbyLeaderChanged&)> OnLobbyLeaderChanged() override;
-	virtual TOnlineEvent<void(const FLobbySchemaChanged&)> OnLobbySchemaChanged() override;
-	virtual TOnlineEvent<void(const FLobbyAttributesChanged&)> OnLobbyAttributesChanged() override;
-	virtual TOnlineEvent<void(const FLobbyMemberAttributesChanged&)> OnLobbyMemberAttributesChanged() override;
-	virtual TOnlineEvent<void(const FLobbyInvitationAdded&)> OnLobbyInvitationAdded() override;
-	virtual TOnlineEvent<void(const FLobbyInvitationRemoved&)> OnLobbyInvitationRemoved() override;
-	virtual TOnlineEvent<void(const FUILobbyJoinRequested&)> OnUILobbyJoinRequested() override;
+	UE_API virtual TOnlineEvent<void(const FLobbyJoined&)> OnLobbyJoined() override;
+	UE_API virtual TOnlineEvent<void(const FLobbyLeft&)> OnLobbyLeft() override;
+	UE_API virtual TOnlineEvent<void(const FLobbyMemberJoined&)> OnLobbyMemberJoined() override;
+	UE_API virtual TOnlineEvent<void(const FLobbyMemberLeft&)> OnLobbyMemberLeft() override;
+	UE_API virtual TOnlineEvent<void(const FLobbyLeaderChanged&)> OnLobbyLeaderChanged() override;
+	UE_API virtual TOnlineEvent<void(const FLobbySchemaChanged&)> OnLobbySchemaChanged() override;
+	UE_API virtual TOnlineEvent<void(const FLobbyAttributesChanged&)> OnLobbyAttributesChanged() override;
+	UE_API virtual TOnlineEvent<void(const FLobbyMemberAttributesChanged&)> OnLobbyMemberAttributesChanged() override;
+	UE_API virtual TOnlineEvent<void(const FLobbyInvitationAdded&)> OnLobbyInvitationAdded() override;
+	UE_API virtual TOnlineEvent<void(const FLobbyInvitationRemoved&)> OnLobbyInvitationRemoved() override;
+	UE_API virtual TOnlineEvent<void(const FUILobbyJoinRequested&)> OnUILobbyJoinRequested() override;
 
 protected:
 #if LOBBIES_FUNCTIONAL_TEST_ENABLED
-	TOnlineAsyncOpHandle<FFunctionalTestLobbies> FunctionalTest(FFunctionalTestLobbies::Params&& Params);
+	UE_API TOnlineAsyncOpHandle<FFunctionalTestLobbies> FunctionalTest(FFunctionalTestLobbies::Params&& Params);
 #endif // LOBBIES_FUNCTIONAL_TEST_ENABLED
 
 	TMap<FAccountId, FLobbyId> PresenceLobbiesUserMap;
@@ -62,3 +64,5 @@ protected:
 };
 
 /* UE::Online */ }
+
+#undef UE_API
diff --git a/Engine/Plugins/Online/OnlineServices/Source/OnlineServicesCommon/Public/Online/LobbiesCommonTypes.h b/Engine/Plugins/Online/OnlineServices/Source/OnlineServicesCommon/Public/Online/LobbiesCommonTypes.h
--- a/Engine/Plugins/Online/OnlineServices/Source/OnlineServicesCommon/Public/Online/LobbiesCommonTypes.h
+++ b/Engine/Plugins/Online/OnlineServices/Source/OnlineServicesCommon/Public/Online/LobbiesCommonTypes.h
@@ -7,6 +7,8 @@
 #include "Online/OnlineAsyncOpHandle.h"
 #include "Online/Schema.h"
 
+#define UE_API ONLINESERVICESCOMMON_API
+
 namespace UE::Online {
 
 #define LOBBIES_FUNCTIONAL_TEST_ENABLED !UE_BUILD_SHIPPING
@@ -268,31 +270,31 @@ struct FLobbyClientDataCommitClientChanges
 };
 
 /** Lobby data as seen by the client. */
-struct ONLINESERVICESCOMMON_API FLobbyClientData final
+struct FLobbyClientData final
 {
 public:
-	FLobbyClientData(FLobbyId LobbyId, const TSharedRef<const FSchemaRegistry>& InSchemaRegistry);
+	UE_API FLobbyClientData(FLobbyId LobbyId, const TSharedRef<const FSchemaRegistry>& InSchemaRegistry);
 
 	TSharedRef<const FLobby> GetPublicDataPtr() const { return InternalPublicData; }
 	const FLobby& GetPublicData() const { return *InternalPublicData; }
 
-	TSharedPtr<const FLobbyMemberInternal> GetMemberData(FAccountId MemberAccountId) const;
+	UE_API TSharedPtr<const FLobbyMemberInternal> GetMemberData(FAccountId MemberAccountId) const;
 
 	/**
 	 * Apply updated lobby data and generate changes.
 	 * The LeaveReason provides context for members who have left since the most recent snapshot.
 	 * Changing a lobby generates events for the local client.
 	 */
-	TOnlineResult<FLobbyClientDataPrepareServiceSnapshot> PrepareServiceSnapshot(FLobbyClientDataPrepareServiceSnapshot::Params&& Params) const;
-	FLobbyClientDataCommitServiceSnapshot::Result CommitServiceSnapshot(FLobbyClientDataCommitServiceSnapshot::Params&& Params);
+	UE_API TOnlineResult<FLobbyClientDataPrepareServiceSnapshot> PrepareServiceSnapshot(FLobbyClientDataPrepareServiceSnapshot::Params&& Params) const;
+	UE_API FLobbyClientDataCommitServiceSnapshot::Result CommitServiceSnapshot(FLobbyClientDataCommitServiceSnapshot::Params&& Params);
 
-	TOnlineResult<FLobbyClientDataPrepareClientChanges> PrepareClientChanges(FLobbyClientDataPrepareClientChanges::Params&& Params) const;
-	FLobbyClientDataCommitClientChanges::Result CommitClientChanges(FLobbyClientDataCommitClientChanges::Params&& Params);
+	UE_API TOnlineResult<FLobbyClientDataPrepareClientChanges> PrepareClientChanges(FLobbyClientDataPrepareClientChanges::Params&& Params) const;
+	UE_API FLobbyClientDataCommitClientChanges::Result CommitClientChanges(FLobbyClientDataCommitClientChanges::Params&& Params);
 
 private:
-	void ResetPreparedChanges() const;
-	void AddMember(const TSharedRef<FLobbyMemberInternal>& Member);
-	void RemoveMember(const TSharedRef<FLobbyMemberInternal>& Member);
+	UE_API void ResetPreparedChanges() const;
+	UE_API void AddMember(const TSharedRef<FLobbyMemberInternal>& Member);
+	UE_API void RemoveMember(const TSharedRef<FLobbyMemberInternal>& Member);
 
 	struct FPreparedServiceChanges
 	{
@@ -385,3 +387,5 @@ END_ONLINE_STRUCT_META()
 /* Meta */ }
 
 /* UE::Online */ }
+
+#undef UE_API
diff --git a/Engine/Plugins/Online/OnlineServices/Source/OnlineServicesCommon/Public/Online/OnlineConfig.h b/Engine/Plugins/Online/OnlineServices/Source/OnlineServicesCommon/Public/Online/OnlineConfig.h
--- a/Engine/Plugins/Online/OnlineServices/Source/OnlineServicesCommon/Public/Online/OnlineConfig.h
+++ b/Engine/Plugins/Online/OnlineServices/Source/OnlineServicesCommon/Public/Online/OnlineConfig.h
@@ -6,6 +6,8 @@
 #include "Online/OnlineMeta.h"
 #include "Online/OnlineServicesLog.h"
 
+#define UE_API ONLINESERVICESCOMMON_API
+
 namespace UE::Online { class FSchemaVariant; }
 
 namespace UE::Online {
@@ -74,7 +76,7 @@ T ParseValueFromString(const FString& StringValue)
 class IOnlineConfigStruct;
 using IOnlineConfigStructPtr = TSharedPtr<IOnlineConfigStruct>;
 
-class ONLINESERVICESCOMMON_API IOnlineConfigStruct
+class IOnlineConfigStruct
 {
 public:
 	virtual ~IOnlineConfigStruct() = default;
@@ -173,7 +175,7 @@ public:
  * be overridden in cases where the underlying configuration system stores integers, floats, etc directly
  * instead of as strings
  */
-class ONLINESERVICESCOMMON_API IOnlineConfigProvider
+class IOnlineConfigProvider
 {
 public:
 	virtual ~IOnlineConfigProvider() {}
@@ -272,15 +274,15 @@ public:
 	}
 };
 
-class ONLINESERVICESCOMMON_API FOnlineConfigStructGConfig : public IOnlineConfigStruct
+class FOnlineConfigStructGConfig : public IOnlineConfigStruct
 {
 public:
 	virtual ~FOnlineConfigStructGConfig() = default;
 
-	virtual bool GetValue(const TCHAR* Key, FString& Value) override;
-	virtual int32 GetValue(const TCHAR* Key, TArray<FString>& Value) override;
-	virtual bool GetValue(const TCHAR * Key, IOnlineConfigStructPtr& Value) override;
-	virtual int32 GetValue(const TCHAR* Key, TArray<IOnlineConfigStructPtr>& Value) override;
+	UE_API virtual bool GetValue(const TCHAR* Key, FString& Value) override;
+	UE_API virtual int32 GetValue(const TCHAR* Key, TArray<FString>& Value) override;
+	UE_API virtual bool GetValue(const TCHAR * Key, IOnlineConfigStructPtr& Value) override;
+	UE_API virtual int32 GetValue(const TCHAR* Key, TArray<IOnlineConfigStructPtr>& Value) override;
 
 private:
 	struct FPrivateToken { explicit FPrivateToken() = default; };
@@ -293,10 +295,10 @@ public:
 private:
 	friend class FOnlineConfigProviderGConfig;
 
-	static IOnlineConfigStructPtr CreateStruct(const FString& InConfigValue);
-	static TArray<IOnlineConfigStructPtr> CreateStructArray(const FString& InConfigValue);
-	static const TCHAR* ParseStruct(const TCHAR* InStr, IOnlineConfigStructPtr& OutPtr);
-	static TArray<FString> CreateValueArray(const FString& InConfigValue);
+	static UE_API IOnlineConfigStructPtr CreateStruct(const FString& InConfigValue);
+	static UE_API TArray<IOnlineConfigStructPtr> CreateStructArray(const FString& InConfigValue);
+	static UE_API const TCHAR* ParseStruct(const TCHAR* InStr, IOnlineConfigStructPtr& OutPtr);
+	static UE_API TArray<FString> CreateValueArray(const FString& InConfigValue);
 
 	TMap<FString, FString> StructMembers;
 };
@@ -304,7 +306,7 @@ private:
 /**
  * IOnlineConfigProvider implementation that uses GConfig to retrieve values
  */
-class ONLINESERVICESCOMMON_API FOnlineConfigProviderGConfig : public IOnlineConfigProvider
+class FOnlineConfigProviderGConfig : public IOnlineConfigProvider
 {
 public:
 	FOnlineConfigProviderGConfig(const FString& InConfigFile)
@@ -313,10 +315,10 @@ public:
 	}
 	virtual ~FOnlineConfigProviderGConfig() = default;
 
-	virtual bool GetValue(const TCHAR* Section, const TCHAR* Key, FString& Value) override;
-	virtual int32 GetValue(const TCHAR* Section, const TCHAR* Key, TArray<FString>& Value) override;
-	virtual bool GetValue(const TCHAR* Section, const TCHAR* Key, IOnlineConfigStructPtr& Value) override;
-	virtual int32 GetValue(const TCHAR* Section, const TCHAR* Key, TArray<IOnlineConfigStructPtr>& Value) override;
+	UE_API virtual bool GetValue(const TCHAR* Section, const TCHAR* Key, FString& Value) override;
+	UE_API virtual int32 GetValue(const TCHAR* Section, const TCHAR* Key, TArray<FString>& Value) override;
+	UE_API virtual bool GetValue(const TCHAR* Section, const TCHAR* Key, IOnlineConfigStructPtr& Value) override;
+	UE_API virtual int32 GetValue(const TCHAR* Section, const TCHAR* Key, TArray<IOnlineConfigStructPtr>& Value) override;
 
 private:
 	FString ConfigFile;
@@ -461,3 +463,5 @@ auto LoadConfig(IOnlineConfigProvider& Provider, const TArray<FString>& SectionH
 }
 
 /* UE::Online */ }
+
+#undef UE_API
diff --git a/Engine/Plugins/Online/OnlineServices/Source/OnlineServicesCommon/Public/Online/OnlineServicesCommon.h b/Engine/Plugins/Online/OnlineServices/Source/OnlineServicesCommon/Public/Online/OnlineServicesCommon.h
--- a/Engine/Plugins/Online/OnlineServices/Source/OnlineServicesCommon/Public/Online/OnlineServicesCommon.h
+++ b/Engine/Plugins/Online/OnlineServices/Source/OnlineServicesCommon/Public/Online/OnlineServicesCommon.h
@@ -11,11 +11,13 @@
 
 #include "Containers/Ticker.h"
 
+#define UE_API ONLINESERVICESCOMMON_API
+
 namespace UE::Online { class IOnlineExecHandler; }
 
 namespace UE::Online {
 
-class ONLINESERVICESCOMMON_API FOnlineServicesCommon
+class FOnlineServicesCommon
 	: public IOnlineServices
 	, public TSharedFromThis<FOnlineServicesCommon>
 	, public FTSTickerObjectBase
@@ -30,33 +32,33 @@ public:
 	{
 	}
 
-	FOnlineServicesCommon(const FString& InServiceConfigName, FName InInstanceName, FName InInstanceConfigName);
+	UE_API FOnlineServicesCommon(const FString& InServiceConfigName, FName InInstanceName, FName InInstanceConfigName);
 	FOnlineServicesCommon(const FOnlineServicesCommon&) = delete;
 	FOnlineServicesCommon(FOnlineServicesCommon&&) = delete;
-	virtual ~FOnlineServicesCommon();
+	UE_API virtual ~FOnlineServicesCommon();
 
 	// IOnlineServices
-	virtual void Init() override;
-	virtual void Destroy() override;
-	virtual IAchievementsPtr GetAchievementsInterface() override;
-	virtual IAuthPtr GetAuthInterface() override;
-	virtual IUserInfoPtr GetUserInfoInterface() override;
-	virtual ICommercePtr GetCommerceInterface() override;
-	virtual ISocialPtr GetSocialInterface() override;
-	virtual IPresencePtr GetPresenceInterface() override;
-	virtual IExternalUIPtr GetExternalUIInterface() override;
-	virtual ILeaderboardsPtr GetLeaderboardsInterface() override;
-	virtual ILobbiesPtr GetLobbiesInterface() override;
-	virtual ISessionsPtr GetSessionsInterface() override;
-	virtual IStatsPtr GetStatsInterface() override;
-	virtual IConnectivityPtr GetConnectivityInterface() override;
-	virtual IPrivilegesPtr GetPrivilegesInterface() override;
-	virtual ITitleFilePtr GetTitleFileInterface() override;
-	virtual IUserFilePtr GetUserFileInterface() override;
-	virtual TOnlineResult<FGetResolvedConnectString> GetResolvedConnectString(FGetResolvedConnectString::Params&& Params) override;
-	virtual FName GetInstanceName() const override;
-	virtual FName GetInstanceConfigName() const override;
-	virtual void AssignBaseInterfaceSharedPtr(const FOnlineTypeName& TypeName, void* OutBaseInterfaceSP) override final;
+	UE_API virtual void Init() override;
+	UE_API virtual void Destroy() override;
+	UE_API virtual IAchievementsPtr GetAchievementsInterface() override;
+	UE_API virtual IAuthPtr GetAuthInterface() override;
+	UE_API virtual IUserInfoPtr GetUserInfoInterface() override;
+	UE_API virtual ICommercePtr GetCommerceInterface() override;
+	UE_API virtual ISocialPtr GetSocialInterface() override;
+	UE_API virtual IPresencePtr GetPresenceInterface() override;
+	UE_API virtual IExternalUIPtr GetExternalUIInterface() override;
+	UE_API virtual ILeaderboardsPtr GetLeaderboardsInterface() override;
+	UE_API virtual ILobbiesPtr GetLobbiesInterface() override;
+	UE_API virtual ISessionsPtr GetSessionsInterface() override;
+	UE_API virtual IStatsPtr GetStatsInterface() override;
+	UE_API virtual IConnectivityPtr GetConnectivityInterface() override;
+	UE_API virtual IPrivilegesPtr GetPrivilegesInterface() override;
+	UE_API virtual ITitleFilePtr GetTitleFileInterface() override;
+	UE_API virtual IUserFilePtr GetUserFileInterface() override;
+	UE_API virtual TOnlineResult<FGetResolvedConnectString> GetResolvedConnectString(FGetResolvedConnectString::Params&& Params) override;
+	UE_API virtual FName GetInstanceName() const override;
+	UE_API virtual FName GetInstanceConfigName() const override;
+	UE_API virtual void AssignBaseInterfaceSharedPtr(const FOnlineTypeName& TypeName, void* OutBaseInterfaceSP) override final;
 
 	// FOnlineServicesCommon
 
@@ -72,37 +74,37 @@ public:
 	/**
 	 * Called to register all the IOnlineComponents with the IOnlineService, called after this is constructed
 	 */
-	virtual void RegisterComponents();
+	UE_API virtual void RegisterComponents();
 
 	/**
 	 * Calls Initialize on all the components, called after RegisterComponents
 	 */
-	virtual void Initialize();
+	UE_API virtual void Initialize();
 
 	/**
 	 * Calls PostInitialize on all the components, called after Initialize
 	 */
-	virtual void PostInitialize();
+	UE_API virtual void PostInitialize();
 
 	/**
 	 * Calls UpdateConfig on all the components
 	 */
-	virtual void UpdateConfig();
+	UE_API virtual void UpdateConfig();
 
 	/**
 	 * Calls Tick on all the components
 	 */
-	virtual bool Tick(float DeltaSeconds) override;
+	UE_API virtual bool Tick(float DeltaSeconds) override;
 
 	/**
 	 * Calls PreShutdown on all the components, called prior to Shutdown
 	 */
-	virtual void PreShutdown();
+	UE_API virtual void PreShutdown();
 
 	/**
 	 * Calls Shutdown on all the components, called before this is destructed
 	 */
-	virtual void Shutdown();
+	UE_API virtual void Shutdown();
 
 	/**
 	 * Call a callable according to a specified execution policy
@@ -330,18 +332,18 @@ public:
 	}
 
 	/* Queue for executing tasks in parallel. Serial queues feed into this */
-	FOnlineAsyncOpQueueParallel& GetParallelQueue();
+	UE_API FOnlineAsyncOpQueueParallel& GetParallelQueue();
 
 	/* Queue for executing tasks in serial */
-	FOnlineAsyncOpQueue& GetSerialQueue();
+	UE_API FOnlineAsyncOpQueue& GetSerialQueue();
 
 	/* Queues for executing per-user tasks in serial */
-	FOnlineAsyncOpQueue& GetSerialQueue(const FAccountId& AccountId);
+	UE_API FOnlineAsyncOpQueue& GetSerialQueue(const FAccountId& AccountId);
 	
-	void RegisterExecHandler(const FString& Name, TUniquePtr<IOnlineExecHandler>&& Handler);
+	UE_API void RegisterExecHandler(const FString& Name, TUniquePtr<IOnlineExecHandler>&& Handler);
 
 #if UE_ALLOW_EXEC_COMMANDS
-	virtual bool Exec(UWorld* World, const TCHAR* Cmd, FOutputDevice& Ar) override;
+	UE_API virtual bool Exec(UWorld* World, const TCHAR* Cmd, FOutputDevice& Ar) override;
 #endif
 
 	FOnlineAsyncOpCache OpCache;
@@ -359,13 +361,13 @@ protected:
 		UE_CLOG(bPreShutdownComplete, LogOnlineServices, Error, TEXT("Can't add op %s anymore after shutdown in %p %s online services!"), OpType::Name, this, LexToString(GetServicesProvider()));
 	}
 
-	void Flush(EAsyncOpFlushReason FlushReason);
-	virtual void FlushTick(float DeltaSeconds);
-	void LoadCommonConfig();
+	UE_API void Flush(EAsyncOpFlushReason FlushReason);
+	UE_API virtual void FlushTick(float DeltaSeconds);
+	UE_API void LoadCommonConfig();
 
 	TMap<FString, TUniquePtr<IOnlineExecHandler>> ExecCommands;
 
-	static uint32 NextInstanceIndex;
+	static UE_API uint32 NextInstanceIndex;
 	uint32 InstanceIndex;
 	FName InstanceName;
 	FName InstanceConfigName;
@@ -384,3 +386,5 @@ protected:
 };
 
 /* UE::Online */ }
+
+#undef UE_API
diff --git a/Engine/Plugins/Online/OnlineServices/Source/OnlineServicesCommon/Public/Online/PresenceCommon.h b/Engine/Plugins/Online/OnlineServices/Source/OnlineServicesCommon/Public/Online/PresenceCommon.h
--- a/Engine/Plugins/Online/OnlineServices/Source/OnlineServicesCommon/Public/Online/PresenceCommon.h
+++ b/Engine/Plugins/Online/OnlineServices/Source/OnlineServicesCommon/Public/Online/PresenceCommon.h
@@ -5,27 +5,29 @@
 #include "Online/Presence.h"
 #include "Online/OnlineComponent.h"
 
+#define UE_API ONLINESERVICESCOMMON_API
+
 namespace UE::Online {
 
 class FOnlineServicesCommon;
 
-class ONLINESERVICESCOMMON_API FPresenceCommon : public TOnlineComponent<IPresence>
+class FPresenceCommon : public TOnlineComponent<IPresence>
 {
 public:
 	using Super = IPresence;
 
-	FPresenceCommon(FOnlineServicesCommon& InServices);
+	UE_API FPresenceCommon(FOnlineServicesCommon& InServices);
 
 	// IOnlineComponent
-	virtual void RegisterCommands() override;
+	UE_API virtual void RegisterCommands() override;
 
 	// IPresence
-	virtual TOnlineAsyncOpHandle<FQueryPresence> QueryPresence(FQueryPresence::Params&& Params) override;
-	virtual TOnlineAsyncOpHandle<FBatchQueryPresence> BatchQueryPresence(FBatchQueryPresence::Params&& Params) override;
-	virtual TOnlineResult<FGetCachedPresence> GetCachedPresence(FGetCachedPresence::Params&& Params) override;
-	virtual TOnlineAsyncOpHandle<FUpdatePresence> UpdatePresence(FUpdatePresence::Params&& Params) override;
-	virtual TOnlineAsyncOpHandle<FPartialUpdatePresence> PartialUpdatePresence(FPartialUpdatePresence::Params&& Params) override;
-	virtual TOnlineEvent<void(const FPresenceUpdated&)> OnPresenceUpdated() override;
+	UE_API virtual TOnlineAsyncOpHandle<FQueryPresence> QueryPresence(FQueryPresence::Params&& Params) override;
+	UE_API virtual TOnlineAsyncOpHandle<FBatchQueryPresence> BatchQueryPresence(FBatchQueryPresence::Params&& Params) override;
+	UE_API virtual TOnlineResult<FGetCachedPresence> GetCachedPresence(FGetCachedPresence::Params&& Params) override;
+	UE_API virtual TOnlineAsyncOpHandle<FUpdatePresence> UpdatePresence(FUpdatePresence::Params&& Params) override;
+	UE_API virtual TOnlineAsyncOpHandle<FPartialUpdatePresence> PartialUpdatePresence(FPartialUpdatePresence::Params&& Params) override;
+	UE_API virtual TOnlineEvent<void(const FPresenceUpdated&)> OnPresenceUpdated() override;
 
 protected:
 	FOnlineServicesCommon& Services;
@@ -34,3 +36,5 @@ protected:
 };
 
 /* UE::Online */ }
+
+#undef UE_API
diff --git a/Engine/Plugins/Online/OnlineServices/Source/OnlineServicesCommon/Public/Online/PrivilegesCommon.h b/Engine/Plugins/Online/OnlineServices/Source/OnlineServicesCommon/Public/Online/PrivilegesCommon.h
--- a/Engine/Plugins/Online/OnlineServices/Source/OnlineServicesCommon/Public/Online/PrivilegesCommon.h
+++ b/Engine/Plugins/Online/OnlineServices/Source/OnlineServicesCommon/Public/Online/PrivilegesCommon.h
@@ -5,24 +5,28 @@
 #include "Online/Privileges.h"
 #include "Online/OnlineComponent.h"
 
+#define UE_API ONLINESERVICESCOMMON_API
+
 namespace UE::Online {
 
 class FOnlineServicesCommon;
 
-class ONLINESERVICESCOMMON_API FPrivilegesCommon : public TOnlineComponent<IPrivileges>
+class FPrivilegesCommon : public TOnlineComponent<IPrivileges>
 {
 public:
 	using Super = IPrivileges;
 
-	FPrivilegesCommon(FOnlineServicesCommon& InServices);
+	UE_API FPrivilegesCommon(FOnlineServicesCommon& InServices);
 
 	// Begin TOnlineComponent
-	virtual void RegisterCommands() override;
+	UE_API virtual void RegisterCommands() override;
 	// End TOnlineComponent
 
 	// Begin IPrivileges
-	virtual TOnlineAsyncOpHandle<FQueryUserPrivilege> QueryUserPrivilege(FQueryUserPrivilege::Params&& Params) override;
+	UE_API virtual TOnlineAsyncOpHandle<FQueryUserPrivilege> QueryUserPrivilege(FQueryUserPrivilege::Params&& Params) override;
 	// End IPrivileges
 };
 
-/* UE::Online */}
\ No newline at end of file
+/* UE::Online */}
+
+#undef UE_API
diff --git a/Engine/Plugins/Online/OnlineServices/Source/OnlineServicesCommon/Public/Online/SessionsCommon.h b/Engine/Plugins/Online/OnlineServices/Source/OnlineServicesCommon/Public/Online/SessionsCommon.h
--- a/Engine/Plugins/Online/OnlineServices/Source/OnlineServicesCommon/Public/Online/SessionsCommon.h
+++ b/Engine/Plugins/Online/OnlineServices/Source/OnlineServicesCommon/Public/Online/SessionsCommon.h
@@ -6,6 +6,8 @@
 #include "Online/OnlineComponent.h"
 #include "Online/OnlineIdCommon.h"
 
+#define UE_API ONLINESERVICESCOMMON_API
+
 namespace UE::Online {
 
 class FOnlineServicesCommon;
@@ -137,7 +139,7 @@ public:
 	TOnlineBasicSessionInviteIdRegistry<FString> BasicRegistry;
 };
 
-class ONLINESERVICESCOMMON_API FSessionCommon : public ISession
+class FSessionCommon : public ISession
 {
 public:
 	FSessionCommon() = default;
@@ -154,14 +156,14 @@ public:
 
 	virtual bool IsJoinable() const override								{ return GetNumOpenConnections() > 0 && SessionSettings.bAllowNewMembers; }
 
-	virtual FString ToLogString() const override;
+	UE_API virtual FString ToLogString() const override;
 
-	virtual void DumpState() const override;
+	UE_API virtual void DumpState() const override;
 
 protected:
-	void DumpMemberData() const;
-	void DumpSessionInfo() const;
-	void DumpSessionSettings() const;
+	UE_API void DumpMemberData() const;
+	UE_API void DumpSessionInfo() const;
+	UE_API void DumpSessionSettings() const;
 
 public:
 	/** Session information that will remain constant throughout the session's lifetime */
@@ -178,21 +180,21 @@ public:
 	/* Set containing user ids for all the session members */
 	FSessionMemberIdsSet SessionMembers;
 
-	FSessionCommon& operator+=(const FSessionUpdate& SessionUpdate);
+	UE_API FSessionCommon& operator+=(const FSessionUpdate& SessionUpdate);
 };
 
-class ONLINESERVICESCOMMON_API FSessionInviteCommon : public ISessionInvite
+class FSessionInviteCommon : public ISessionInvite
 {
 public:
 	FSessionInviteCommon() = default;
-	FSessionInviteCommon(const FAccountId& InRecipientId, const FAccountId& InSenderId, const FSessionInviteId& InInviteId, const FOnlineSessionId& InSessionId);
+	UE_API FSessionInviteCommon(const FAccountId& InRecipientId, const FAccountId& InSenderId, const FSessionInviteId& InInviteId, const FOnlineSessionId& InSessionId);
 
 	virtual const FAccountId GetRecipientId() const override { return RecipientId; }
 	virtual const FAccountId GetSenderId() const override { return SenderId; }
 	virtual const FSessionInviteId GetInviteId() const override { return InviteId; }
 	virtual const FOnlineSessionId GetSessionId() const override { return SessionId; }
 
-	virtual FString ToLogString() const override;
+	UE_API virtual FString ToLogString() const override;
 
 public:
 	/* The id handle for the user which the invite got sent to */
@@ -238,106 +240,106 @@ struct FGetMutableSessionById
 	};
 };
 
-class ONLINESERVICESCOMMON_API FSessionsCommon : public TOnlineComponent<ISessions>
+class FSessionsCommon : public TOnlineComponent<ISessions>
 {
 public:
 	using Super = ISessions;
 
-	FSessionsCommon(FOnlineServicesCommon& InServices);
+	UE_API FSessionsCommon(FOnlineServicesCommon& InServices);
 
 	// TOnlineComponent
-	virtual void Initialize() override;
-	virtual void RegisterCommands() override;
+	UE_API virtual void Initialize() override;
+	UE_API virtual void RegisterCommands() override;
 
 	// ISessions
-	virtual TOnlineResult<FGetAllSessions> GetAllSessions(FGetAllSessions::Params&& Params) const override;
-	virtual TOnlineResult<FGetSessionByName> GetSessionByName(FGetSessionByName::Params&& Params) const override;
-	virtual TOnlineResult<FGetSessionById> GetSessionById(FGetSessionById::Params&& Params) const override;
-	virtual TOnlineResult<FGetPresenceSession> GetPresenceSession(FGetPresenceSession::Params&& Params) const override;
-	virtual TOnlineResult<FIsPresenceSession> IsPresenceSession(FIsPresenceSession::Params&& Params) const override;
-	virtual TOnlineResult<FSetPresenceSession> SetPresenceSession(FSetPresenceSession::Params&& Params) override;
-	virtual TOnlineResult<FClearPresenceSession> ClearPresenceSession(FClearPresenceSession::Params&& Params) override;
-	virtual TOnlineAsyncOpHandle<FCreateSession> CreateSession(FCreateSession::Params&& Params) override;
-	virtual TOnlineAsyncOpHandle<FUpdateSessionSettings> UpdateSessionSettings(FUpdateSessionSettings::Params&& Params) override;
-	virtual TOnlineAsyncOpHandle<FLeaveSession> LeaveSession(FLeaveSession::Params&& Params) override;
-	virtual TOnlineAsyncOpHandle<FFindSessions> FindSessions(FFindSessions::Params&& Params) override;
-	virtual TOnlineAsyncOpHandle<FStartMatchmaking> StartMatchmaking(FStartMatchmaking::Params&& Params) override;
-	virtual TOnlineAsyncOpHandle<FJoinSession> JoinSession(FJoinSession::Params&& Params) override;
-	virtual TOnlineAsyncOpHandle<FAddSessionMember> AddSessionMember(FAddSessionMember::Params&& Params) override;
-	virtual TOnlineAsyncOpHandle<FRemoveSessionMember> RemoveSessionMember(FRemoveSessionMember::Params&& Params) override;
-	virtual TOnlineAsyncOpHandle<FSendSessionInvite> SendSessionInvite(FSendSessionInvite::Params&& Params) override;
-	virtual TOnlineResult<FGetSessionInviteById> GetSessionInviteById(FGetSessionInviteById::Params&& Params) override;
-	virtual TOnlineResult<FGetAllSessionInvites> GetAllSessionInvites(FGetAllSessionInvites::Params&& Params) override;
-	virtual TOnlineAsyncOpHandle<FRejectSessionInvite> RejectSessionInvite(FRejectSessionInvite::Params&& Params) override;
-
-	virtual TOnlineEvent<void(const FSessionCreated&)> OnSessionCreated() override;
-	virtual TOnlineEvent<void(const FSessionJoined&)> OnSessionJoined() override;
-	virtual TOnlineEvent<void(const FSessionLeft&)> OnSessionLeft() override;
-	virtual TOnlineEvent<void(const FSessionUpdated&)> OnSessionUpdated() override;
-	virtual TOnlineEvent<void(const FSessionInviteReceived&)> OnSessionInviteReceived() override;
-	virtual TOnlineEvent<void(const FUISessionJoinRequested&)> OnUISessionJoinRequested() override;
+	UE_API virtual TOnlineResult<FGetAllSessions> GetAllSessions(FGetAllSessions::Params&& Params) const override;
+	UE_API virtual TOnlineResult<FGetSessionByName> GetSessionByName(FGetSessionByName::Params&& Params) const override;
+	UE_API virtual TOnlineResult<FGetSessionById> GetSessionById(FGetSessionById::Params&& Params) const override;
+	UE_API virtual TOnlineResult<FGetPresenceSession> GetPresenceSession(FGetPresenceSession::Params&& Params) const override;
+	UE_API virtual TOnlineResult<FIsPresenceSession> IsPresenceSession(FIsPresenceSession::Params&& Params) const override;
+	UE_API virtual TOnlineResult<FSetPresenceSession> SetPresenceSession(FSetPresenceSession::Params&& Params) override;
+	UE_API virtual TOnlineResult<FClearPresenceSession> ClearPresenceSession(FClearPresenceSession::Params&& Params) override;
+	UE_API virtual TOnlineAsyncOpHandle<FCreateSession> CreateSession(FCreateSession::Params&& Params) override;
+	UE_API virtual TOnlineAsyncOpHandle<FUpdateSessionSettings> UpdateSessionSettings(FUpdateSessionSettings::Params&& Params) override;
+	UE_API virtual TOnlineAsyncOpHandle<FLeaveSession> LeaveSession(FLeaveSession::Params&& Params) override;
+	UE_API virtual TOnlineAsyncOpHandle<FFindSessions> FindSessions(FFindSessions::Params&& Params) override;
+	UE_API virtual TOnlineAsyncOpHandle<FStartMatchmaking> StartMatchmaking(FStartMatchmaking::Params&& Params) override;
+	UE_API virtual TOnlineAsyncOpHandle<FJoinSession> JoinSession(FJoinSession::Params&& Params) override;
+	UE_API virtual TOnlineAsyncOpHandle<FAddSessionMember> AddSessionMember(FAddSessionMember::Params&& Params) override;
+	UE_API virtual TOnlineAsyncOpHandle<FRemoveSessionMember> RemoveSessionMember(FRemoveSessionMember::Params&& Params) override;
+	UE_API virtual TOnlineAsyncOpHandle<FSendSessionInvite> SendSessionInvite(FSendSessionInvite::Params&& Params) override;
+	UE_API virtual TOnlineResult<FGetSessionInviteById> GetSessionInviteById(FGetSessionInviteById::Params&& Params) override;
+	UE_API virtual TOnlineResult<FGetAllSessionInvites> GetAllSessionInvites(FGetAllSessionInvites::Params&& Params) override;
+	UE_API virtual TOnlineAsyncOpHandle<FRejectSessionInvite> RejectSessionInvite(FRejectSessionInvite::Params&& Params) override;
+
+	UE_API virtual TOnlineEvent<void(const FSessionCreated&)> OnSessionCreated() override;
+	UE_API virtual TOnlineEvent<void(const FSessionJoined&)> OnSessionJoined() override;
+	UE_API virtual TOnlineEvent<void(const FSessionLeft&)> OnSessionLeft() override;
+	UE_API virtual TOnlineEvent<void(const FSessionUpdated&)> OnSessionUpdated() override;
+	UE_API virtual TOnlineEvent<void(const FSessionInviteReceived&)> OnSessionInviteReceived() override;
+	UE_API virtual TOnlineEvent<void(const FUISessionJoinRequested&)> OnUISessionJoinRequested() override;
 
 protected:
 	// Auxiliary methods
-	TOnlineResult<FGetMutableSessionByName> GetMutableSessionByName(FGetMutableSessionByName::Params&& Params) const;
-	TOnlineResult<FGetMutableSessionById> GetMutableSessionById(FGetMutableSessionById::Params&& Params) const;
+	UE_API TOnlineResult<FGetMutableSessionByName> GetMutableSessionByName(FGetMutableSessionByName::Params&& Params) const;
+	UE_API TOnlineResult<FGetMutableSessionById> GetMutableSessionById(FGetMutableSessionById::Params&& Params) const;
 
-	void AddSessionInvite(const TSharedRef<FSessionInviteCommon> SessionInvite, const TSharedRef<FSessionCommon> Session, const FAccountId& LocalAccountId);
-	void AddSearchResult(const TSharedRef<FSessionCommon> Session, const FAccountId& LocalAccountId);
-	void AddSessionWithReferences(const TSharedRef<FSessionCommon> Session, const FName& SessionName, const FAccountId& LocalAccountId, bool bIsPresenceSession);
-	void AddSessionReferences(const FOnlineSessionId SessionId, const FName& SessionName, const FAccountId& LocalAccountId, bool bIsPresenceSession);
+	UE_API void AddSessionInvite(const TSharedRef<FSessionInviteCommon> SessionInvite, const TSharedRef<FSessionCommon> Session, const FAccountId& LocalAccountId);
+	UE_API void AddSearchResult(const TSharedRef<FSessionCommon> Session, const FAccountId& LocalAccountId);
+	UE_API void AddSessionWithReferences(const TSharedRef<FSessionCommon> Session, const FName& SessionName, const FAccountId& LocalAccountId, bool bIsPresenceSession);
+	UE_API void AddSessionReferences(const FOnlineSessionId SessionId, const FName& SessionName, const FAccountId& LocalAccountId, bool bIsPresenceSession);
 
-	void ClearSessionInvitesForSession(const FAccountId& LocalAccountId, const FOnlineSessionId SessionId);
-	void ClearSessionReferences(const FOnlineSessionId SessionId, const FName& SessionName, const FAccountId& LocalAccountId);
+	UE_API void ClearSessionInvitesForSession(const FAccountId& LocalAccountId, const FOnlineSessionId SessionId);
+	UE_API void ClearSessionReferences(const FOnlineSessionId SessionId, const FName& SessionName, const FAccountId& LocalAccountId);
 
-	FSessionUpdate BuildSessionUpdate(const TSharedRef<FSessionCommon>& Session, const FSessionSettingsUpdate& UpdatedValues) const;
+	UE_API FSessionUpdate BuildSessionUpdate(const TSharedRef<FSessionCommon>& Session, const FSessionSettingsUpdate& UpdatedValues) const;
 
 	// FSessionsCommon internal interface
 
-	virtual TOnlineResult<FSetPresenceSession> SetPresenceSessionImpl(FSetPresenceSession::Params&& Params);
-	virtual TOnlineResult<FClearPresenceSession> ClearPresenceSessionImpl(FClearPresenceSession::Params&& Params);
+	UE_API virtual TOnlineResult<FSetPresenceSession> SetPresenceSessionImpl(FSetPresenceSession::Params&& Params);
+	UE_API virtual TOnlineResult<FClearPresenceSession> ClearPresenceSessionImpl(FClearPresenceSession::Params&& Params);
 
-	virtual TFuture<TOnlineResult<FCreateSession>> CreateSessionImpl(const FCreateSession::Params& Params);
-	virtual TOptional<FOnlineError> CheckParams(const FCreateSession::Params& Params) const;
-	virtual TOptional<FOnlineError> CheckState(const FCreateSession::Params& Params) const;
+	UE_API virtual TFuture<TOnlineResult<FCreateSession>> CreateSessionImpl(const FCreateSession::Params& Params);
+	UE_API virtual TOptional<FOnlineError> CheckParams(const FCreateSession::Params& Params) const;
+	UE_API virtual TOptional<FOnlineError> CheckState(const FCreateSession::Params& Params) const;
 
-	virtual TFuture<TOnlineResult<FUpdateSessionSettings>> UpdateSessionSettingsImpl(const FUpdateSessionSettings::Params& Params);
-	virtual TOptional<FOnlineError> CheckParams(const FUpdateSessionSettings::Params& Params) const;
-	virtual TOptional<FOnlineError> CheckState(const FUpdateSessionSettings::Params& Params) const;
+	UE_API virtual TFuture<TOnlineResult<FUpdateSessionSettings>> UpdateSessionSettingsImpl(const FUpdateSessionSettings::Params& Params);
+	UE_API virtual TOptional<FOnlineError> CheckParams(const FUpdateSessionSettings::Params& Params) const;
+	UE_API virtual TOptional<FOnlineError> CheckState(const FUpdateSessionSettings::Params& Params) const;
 
-	virtual TFuture<TOnlineResult<FLeaveSession>> LeaveSessionImpl(const FLeaveSession::Params& Params);
-	virtual TOptional<FOnlineError> CheckParams(const FLeaveSession::Params& Params) const;
-	virtual TOptional<FOnlineError> CheckState(const FLeaveSession::Params& Params) const;
+	UE_API virtual TFuture<TOnlineResult<FLeaveSession>> LeaveSessionImpl(const FLeaveSession::Params& Params);
+	UE_API virtual TOptional<FOnlineError> CheckParams(const FLeaveSession::Params& Params) const;
+	UE_API virtual TOptional<FOnlineError> CheckState(const FLeaveSession::Params& Params) const;
 
-	virtual TFuture<TOnlineResult<FFindSessions>> FindSessionsImpl(const FFindSessions::Params& Params);
-	virtual TOptional<FOnlineError> CheckParams(const FFindSessions::Params& Params) const;
-	virtual TOptional<FOnlineError> CheckState(const FFindSessions::Params& Params) const;
+	UE_API virtual TFuture<TOnlineResult<FFindSessions>> FindSessionsImpl(const FFindSessions::Params& Params);
+	UE_API virtual TOptional<FOnlineError> CheckParams(const FFindSessions::Params& Params) const;
+	UE_API virtual TOptional<FOnlineError> CheckState(const FFindSessions::Params& Params) const;
 
-	virtual TFuture<TOnlineResult<FStartMatchmaking>> StartMatchmakingImpl(const FStartMatchmaking::Params& Params);
-	virtual TOptional<FOnlineError> CheckParams(const FStartMatchmaking::Params& Params) const;
-	virtual TOptional<FOnlineError> CheckState(const FStartMatchmaking::Params& Params) const;
+	UE_API virtual TFuture<TOnlineResult<FStartMatchmaking>> StartMatchmakingImpl(const FStartMatchmaking::Params& Params);
+	UE_API virtual TOptional<FOnlineError> CheckParams(const FStartMatchmaking::Params& Params) const;
+	UE_API virtual TOptional<FOnlineError> CheckState(const FStartMatchmaking::Params& Params) const;
 
-	virtual TFuture<TOnlineResult<FJoinSession>> JoinSessionImpl(const FJoinSession::Params& Params);
-	virtual TOptional<FOnlineError> CheckParams(const FJoinSession::Params& Params) const;
-	virtual TOptional<FOnlineError> CheckState(const FJoinSession::Params& Params) const;
+	UE_API virtual TFuture<TOnlineResult<FJoinSession>> JoinSessionImpl(const FJoinSession::Params& Params);
+	UE_API virtual TOptional<FOnlineError> CheckParams(const FJoinSession::Params& Params) const;
+	UE_API virtual TOptional<FOnlineError> CheckState(const FJoinSession::Params& Params) const;
 
-	virtual TFuture<TOnlineResult<FAddSessionMember>> AddSessionMemberImpl(const FAddSessionMember::Params& Params);
-	virtual TOptional<FOnlineError> CheckParams(const FAddSessionMember::Params& Params) const;
-	virtual TOptional<FOnlineError> CheckState(const FAddSessionMember::Params& Params) const;
-	TOnlineResult<FAddSessionMember> AddSessionMemberInternal(const FAddSessionMember::Params& Params);
+	UE_API virtual TFuture<TOnlineResult<FAddSessionMember>> AddSessionMemberImpl(const FAddSessionMember::Params& Params);
+	UE_API virtual TOptional<FOnlineError> CheckParams(const FAddSessionMember::Params& Params) const;
+	UE_API virtual TOptional<FOnlineError> CheckState(const FAddSessionMember::Params& Params) const;
+	UE_API TOnlineResult<FAddSessionMember> AddSessionMemberInternal(const FAddSessionMember::Params& Params);
 
-	virtual TFuture<TOnlineResult<FRemoveSessionMember>> RemoveSessionMemberImpl(const FRemoveSessionMember::Params& Params);
-	virtual TOptional<FOnlineError> CheckParams(const FRemoveSessionMember::Params& Params) const;
-	virtual TOptional<FOnlineError> CheckState(const FRemoveSessionMember::Params& Params) const;
-	TOnlineResult<FRemoveSessionMember> RemoveSessionMemberInternal(const FRemoveSessionMember::Params& Params);
+	UE_API virtual TFuture<TOnlineResult<FRemoveSessionMember>> RemoveSessionMemberImpl(const FRemoveSessionMember::Params& Params);
+	UE_API virtual TOptional<FOnlineError> CheckParams(const FRemoveSessionMember::Params& Params) const;
+	UE_API virtual TOptional<FOnlineError> CheckState(const FRemoveSessionMember::Params& Params) const;
+	UE_API TOnlineResult<FRemoveSessionMember> RemoveSessionMemberInternal(const FRemoveSessionMember::Params& Params);
 
-	virtual TFuture<TOnlineResult<FSendSessionInvite>> SendSessionInviteImpl(const FSendSessionInvite::Params& Params);
-	virtual TOptional<FOnlineError> CheckParams(const FSendSessionInvite::Params& Params) const;
-	virtual TOptional<FOnlineError> CheckState(const FSendSessionInvite::Params& Params) const;
+	UE_API virtual TFuture<TOnlineResult<FSendSessionInvite>> SendSessionInviteImpl(const FSendSessionInvite::Params& Params);
+	UE_API virtual TOptional<FOnlineError> CheckParams(const FSendSessionInvite::Params& Params) const;
+	UE_API virtual TOptional<FOnlineError> CheckState(const FSendSessionInvite::Params& Params) const;
 
-	virtual TFuture<TOnlineResult<FRejectSessionInvite>> RejectSessionInviteImpl(const FRejectSessionInvite::Params& Params);
-	virtual TOptional<FOnlineError> CheckParams(const FRejectSessionInvite::Params& Params) const;
-	virtual TOptional<FOnlineError> CheckState(const FRejectSessionInvite::Params& Params) const;
+	UE_API virtual TFuture<TOnlineResult<FRejectSessionInvite>> RejectSessionInviteImpl(const FRejectSessionInvite::Params& Params);
+	UE_API virtual TOptional<FOnlineError> CheckParams(const FRejectSessionInvite::Params& Params) const;
+	UE_API virtual TOptional<FOnlineError> CheckState(const FRejectSessionInvite::Params& Params) const;
 
 private:
 	template<typename MethodStruct>
@@ -401,3 +403,5 @@ END_ONLINE_STRUCT_META()
 /* Meta*/ }
 
 /* UE::Online */ }
+
+#undef UE_API
diff --git a/Engine/Plugins/Online/OnlineServices/Source/OnlineServicesCommon/Public/Online/SocialCommon.h b/Engine/Plugins/Online/OnlineServices/Source/OnlineServicesCommon/Public/Online/SocialCommon.h
--- a/Engine/Plugins/Online/OnlineServices/Source/OnlineServicesCommon/Public/Online/SocialCommon.h
+++ b/Engine/Plugins/Online/OnlineServices/Source/OnlineServicesCommon/Public/Online/SocialCommon.h
@@ -5,34 +5,38 @@
 #include "Online/Social.h"
 #include "Online/OnlineComponent.h"
 
+#define UE_API ONLINESERVICESCOMMON_API
+
 namespace UE::Online {
 
 class FOnlineServicesCommon;
 
-class ONLINESERVICESCOMMON_API FSocialCommon : public TOnlineComponent<ISocial>
+class FSocialCommon : public TOnlineComponent<ISocial>
 {
 public:
 	using Super = ISocial;
 
-	FSocialCommon(FOnlineServicesCommon& InServices);
+	UE_API FSocialCommon(FOnlineServicesCommon& InServices);
 
 	// TOnlineComponent
-	virtual void RegisterCommands() override;
+	UE_API virtual void RegisterCommands() override;
 
 	// ISocial
-	virtual TOnlineAsyncOpHandle<FQueryFriends> QueryFriends(FQueryFriends::Params&& Params) override;
-	virtual TOnlineResult<FGetFriends> GetFriends(FGetFriends::Params&& Params) override;
-	virtual TOnlineAsyncOpHandle<FSendFriendInvite> SendFriendInvite(FSendFriendInvite::Params&& Params) override;
-	virtual TOnlineAsyncOpHandle<FAcceptFriendInvite> AcceptFriendInvite(FAcceptFriendInvite::Params&& Params) override;
-	virtual TOnlineAsyncOpHandle<FRejectFriendInvite> RejectFriendInvite(FRejectFriendInvite::Params&& Params) override;
-	virtual TOnlineEvent<void(const FRelationshipUpdated&)> OnRelationshipUpdated() override;
-	virtual TOnlineAsyncOpHandle<FQueryBlockedUsers> QueryBlockedUsers(FQueryBlockedUsers::Params&& Params) override;
-	virtual TOnlineResult<FGetBlockedUsers> GetBlockedUsers(FGetBlockedUsers::Params&& Params) override;
-	virtual TOnlineAsyncOpHandle<FBlockUser> BlockUser(FBlockUser::Params&& Params) override;
+	UE_API virtual TOnlineAsyncOpHandle<FQueryFriends> QueryFriends(FQueryFriends::Params&& Params) override;
+	UE_API virtual TOnlineResult<FGetFriends> GetFriends(FGetFriends::Params&& Params) override;
+	UE_API virtual TOnlineAsyncOpHandle<FSendFriendInvite> SendFriendInvite(FSendFriendInvite::Params&& Params) override;
+	UE_API virtual TOnlineAsyncOpHandle<FAcceptFriendInvite> AcceptFriendInvite(FAcceptFriendInvite::Params&& Params) override;
+	UE_API virtual TOnlineAsyncOpHandle<FRejectFriendInvite> RejectFriendInvite(FRejectFriendInvite::Params&& Params) override;
+	UE_API virtual TOnlineEvent<void(const FRelationshipUpdated&)> OnRelationshipUpdated() override;
+	UE_API virtual TOnlineAsyncOpHandle<FQueryBlockedUsers> QueryBlockedUsers(FQueryBlockedUsers::Params&& Params) override;
+	UE_API virtual TOnlineResult<FGetBlockedUsers> GetBlockedUsers(FGetBlockedUsers::Params&& Params) override;
+	UE_API virtual TOnlineAsyncOpHandle<FBlockUser> BlockUser(FBlockUser::Params&& Params) override;
 
 protected:
-	void BroadcastRelationshipUpdated(FAccountId LocalAccountId, FAccountId RemoteAccountId, ERelationship OldRelationship, ERelationship NewRelationship);
+	UE_API void BroadcastRelationshipUpdated(FAccountId LocalAccountId, FAccountId RemoteAccountId, ERelationship OldRelationship, ERelationship NewRelationship);
 	TOnlineEventCallable<void(const FRelationshipUpdated&)> OnRelationshipUpdatedEvent;
 };
 
 /* UE::Online */ }
+
+#undef UE_API
diff --git a/Engine/Plugins/Online/OnlineServices/Source/OnlineServicesCommon/Public/Online/StatsCommon.h b/Engine/Plugins/Online/OnlineServices/Source/OnlineServicesCommon/Public/Online/StatsCommon.h
--- a/Engine/Plugins/Online/OnlineServices/Source/OnlineServicesCommon/Public/Online/StatsCommon.h
+++ b/Engine/Plugins/Online/OnlineServices/Source/OnlineServicesCommon/Public/Online/StatsCommon.h
@@ -5,6 +5,8 @@
 #include "Online/Stats.h"
 #include "Online/OnlineComponent.h"
 
+#define UE_API ONLINESERVICESCOMMON_API
+
 namespace UE::Online {
 
 enum class EStatModifyMethod : uint8
@@ -69,31 +71,31 @@ struct FFindUserStatsByAccountId
 	FAccountId AccountId;
 };
 
-class ONLINESERVICESCOMMON_API FStatsCommon : public TOnlineComponent<IStats>
+class FStatsCommon : public TOnlineComponent<IStats>
 {
 public:
 	using Super = IStats;
 
-	FStatsCommon(FOnlineServicesCommon& InServices);
+	UE_API FStatsCommon(FOnlineServicesCommon& InServices);
 
 	// TOnlineComponent
-	virtual void UpdateConfig() override;
-	virtual void RegisterCommands() override;
+	UE_API virtual void UpdateConfig() override;
+	UE_API virtual void RegisterCommands() override;
 
 	// IStats
-	virtual TOnlineAsyncOpHandle<FUpdateStats> UpdateStats(FUpdateStats::Params&& Params) override;
-	virtual TOnlineAsyncOpHandle<FQueryStats> QueryStats(FQueryStats::Params&& Params) override;
-	virtual TOnlineAsyncOpHandle<FBatchQueryStats> BatchQueryStats(FBatchQueryStats::Params&& Params) override;
+	UE_API virtual TOnlineAsyncOpHandle<FUpdateStats> UpdateStats(FUpdateStats::Params&& Params) override;
+	UE_API virtual TOnlineAsyncOpHandle<FQueryStats> QueryStats(FQueryStats::Params&& Params) override;
+	UE_API virtual TOnlineAsyncOpHandle<FBatchQueryStats> BatchQueryStats(FBatchQueryStats::Params&& Params) override;
 #if !UE_BUILD_SHIPPING
-	virtual TOnlineAsyncOpHandle<FResetStats> ResetStats(FResetStats::Params&& Params) override;
+	UE_API virtual TOnlineAsyncOpHandle<FResetStats> ResetStats(FResetStats::Params&& Params) override;
 #endif // !UE_BUILD_SHIPPING
-	virtual TOnlineResult<FGetCachedStats> GetCachedStats(FGetCachedStats::Params&& Params) const override;
+	UE_API virtual TOnlineResult<FGetCachedStats> GetCachedStats(FGetCachedStats::Params&& Params) const override;
 
-	virtual TOnlineEvent<void(const FStatsUpdated&)> OnStatsUpdated() override;
-	const FStatDefinition* GetStatDefinition(const FString& StatName) const;
+	UE_API virtual TOnlineEvent<void(const FStatsUpdated&)> OnStatsUpdated() override;
+	UE_API const FStatDefinition* GetStatDefinition(const FString& StatName) const;
 
 protected:
-	void CacheUserStats(const FUserStats& UserStats);
+	UE_API void CacheUserStats(const FUserStats& UserStats);
 
 	TMap<FString, FStatDefinition> StatDefinitions;
 	TOnlineEventCallable<void(const FStatsUpdated& StatsUpdated)> OnStatsUpdatedEvent;
@@ -102,3 +104,5 @@ protected:
 };
 
 /* UE::Online */ }
+
+#undef UE_API
diff --git a/Engine/Plugins/Online/OnlineServices/Source/OnlineServicesCommon/Public/Online/TitleFileCommon.h b/Engine/Plugins/Online/OnlineServices/Source/OnlineServicesCommon/Public/Online/TitleFileCommon.h
--- a/Engine/Plugins/Online/OnlineServices/Source/OnlineServicesCommon/Public/Online/TitleFileCommon.h
+++ b/Engine/Plugins/Online/OnlineServices/Source/OnlineServicesCommon/Public/Online/TitleFileCommon.h
@@ -5,24 +5,28 @@
 #include "Online/TitleFile.h"
 #include "Online/OnlineComponent.h"
 
+#define UE_API ONLINESERVICESCOMMON_API
+
 namespace UE::Online {
 
 class FOnlineServicesCommon;
 
-class ONLINESERVICESCOMMON_API FTitleFileCommon : public TOnlineComponent<ITitleFile>
+class FTitleFileCommon : public TOnlineComponent<ITitleFile>
 {
 public:
 	using Super = ITitleFile;
 
-	FTitleFileCommon(FOnlineServicesCommon& InServices);
+	UE_API FTitleFileCommon(FOnlineServicesCommon& InServices);
 
 	// TOnlineComponent
-	virtual void RegisterCommands() override;
+	UE_API virtual void RegisterCommands() override;
 
 	// ITitleFile
-	virtual TOnlineAsyncOpHandle<FTitleFileEnumerateFiles> EnumerateFiles(FTitleFileEnumerateFiles::Params&& Params) override;
-	virtual TOnlineResult<FTitleFileGetEnumeratedFiles> GetEnumeratedFiles(FTitleFileGetEnumeratedFiles::Params&& Params) override;
-	virtual TOnlineAsyncOpHandle<FTitleFileReadFile> ReadFile(FTitleFileReadFile::Params&& Params) override;
+	UE_API virtual TOnlineAsyncOpHandle<FTitleFileEnumerateFiles> EnumerateFiles(FTitleFileEnumerateFiles::Params&& Params) override;
+	UE_API virtual TOnlineResult<FTitleFileGetEnumeratedFiles> GetEnumeratedFiles(FTitleFileGetEnumeratedFiles::Params&& Params) override;
+	UE_API virtual TOnlineAsyncOpHandle<FTitleFileReadFile> ReadFile(FTitleFileReadFile::Params&& Params) override;
 };
 
-/* UE::Online */}
\ No newline at end of file
+/* UE::Online */}
+
+#undef UE_API
diff --git a/Engine/Plugins/Online/OnlineServices/Source/OnlineServicesCommon/Public/Online/UserFileCommon.h b/Engine/Plugins/Online/OnlineServices/Source/OnlineServicesCommon/Public/Online/UserFileCommon.h
--- a/Engine/Plugins/Online/OnlineServices/Source/OnlineServicesCommon/Public/Online/UserFileCommon.h
+++ b/Engine/Plugins/Online/OnlineServices/Source/OnlineServicesCommon/Public/Online/UserFileCommon.h
@@ -5,27 +5,31 @@
 #include "Online/UserFile.h"
 #include "Online/OnlineComponent.h"
 
+#define UE_API ONLINESERVICESCOMMON_API
+
 namespace UE::Online {
 
 class FOnlineServicesCommon;
 
-class ONLINESERVICESCOMMON_API FUserFileCommon : public TOnlineComponent<IUserFile>
+class FUserFileCommon : public TOnlineComponent<IUserFile>
 {
 public:
 	using Super = IUserFile;
 
-	FUserFileCommon(FOnlineServicesCommon& InServices);
+	UE_API FUserFileCommon(FOnlineServicesCommon& InServices);
 
 	// TOnlineComponent
-	virtual void RegisterCommands() override;
+	UE_API virtual void RegisterCommands() override;
 
 	// IUserFile
-	virtual TOnlineAsyncOpHandle<FUserFileEnumerateFiles> EnumerateFiles(FUserFileEnumerateFiles::Params&& Params) override;
-	virtual TOnlineResult<FUserFileGetEnumeratedFiles> GetEnumeratedFiles(FUserFileGetEnumeratedFiles::Params&& Params) override;
-	virtual TOnlineAsyncOpHandle<FUserFileReadFile> ReadFile(FUserFileReadFile::Params&& Params) override;
-	virtual TOnlineAsyncOpHandle<FUserFileWriteFile> WriteFile(FUserFileWriteFile::Params&& Params) override;
-	virtual TOnlineAsyncOpHandle<FUserFileCopyFile> CopyFile(FUserFileCopyFile::Params&& Params) override;
-	virtual TOnlineAsyncOpHandle<FUserFileDeleteFile> DeleteFile(FUserFileDeleteFile::Params&& Params) override;
+	UE_API virtual TOnlineAsyncOpHandle<FUserFileEnumerateFiles> EnumerateFiles(FUserFileEnumerateFiles::Params&& Params) override;
+	UE_API virtual TOnlineResult<FUserFileGetEnumeratedFiles> GetEnumeratedFiles(FUserFileGetEnumeratedFiles::Params&& Params) override;
+	UE_API virtual TOnlineAsyncOpHandle<FUserFileReadFile> ReadFile(FUserFileReadFile::Params&& Params) override;
+	UE_API virtual TOnlineAsyncOpHandle<FUserFileWriteFile> WriteFile(FUserFileWriteFile::Params&& Params) override;
+	UE_API virtual TOnlineAsyncOpHandle<FUserFileCopyFile> CopyFile(FUserFileCopyFile::Params&& Params) override;
+	UE_API virtual TOnlineAsyncOpHandle<FUserFileDeleteFile> DeleteFile(FUserFileDeleteFile::Params&& Params) override;
 };
 
-/* UE::Online */}
\ No newline at end of file
+/* UE::Online */}
+
+#undef UE_API
diff --git a/Engine/Plugins/Online/OnlineServices/Source/OnlineServicesCommon/Public/Online/UserInfoCommon.h b/Engine/Plugins/Online/OnlineServices/Source/OnlineServicesCommon/Public/Online/UserInfoCommon.h
--- a/Engine/Plugins/Online/OnlineServices/Source/OnlineServicesCommon/Public/Online/UserInfoCommon.h
+++ b/Engine/Plugins/Online/OnlineServices/Source/OnlineServicesCommon/Public/Online/UserInfoCommon.h
@@ -5,26 +5,30 @@
 #include "Online/UserInfo.h"
 #include "Online/OnlineComponent.h"
 
+#define UE_API ONLINESERVICESCOMMON_API
+
 namespace UE::Online {
 
 class FOnlineServicesCommon;
 
-class ONLINESERVICESCOMMON_API FUserInfoCommon : public TOnlineComponent<IUserInfo>
+class FUserInfoCommon : public TOnlineComponent<IUserInfo>
 {
 public:
 	using Super = IUserInfo;
 
-	FUserInfoCommon(FOnlineServicesCommon& InServices);
+	UE_API FUserInfoCommon(FOnlineServicesCommon& InServices);
 
 	// TOnlineComponent
-	virtual void RegisterCommands() override;
+	UE_API virtual void RegisterCommands() override;
 
 	// IUserInfo
-	virtual TOnlineAsyncOpHandle<FQueryUserInfo> QueryUserInfo(FQueryUserInfo::Params&& Params) override;
-	virtual TOnlineResult<FGetUserInfo> GetUserInfo(FGetUserInfo::Params&& Params) override;
-	virtual TOnlineAsyncOpHandle<FQueryUserAvatar> QueryUserAvatar(FQueryUserAvatar::Params&& Params) override;
-	virtual TOnlineResult<FGetUserAvatar> GetUserAvatar(FGetUserAvatar::Params&& Params) override;
-	virtual TOnlineAsyncOpHandle<FShowUserProfile> ShowUserProfile(FShowUserProfile::Params&& Params) override;
+	UE_API virtual TOnlineAsyncOpHandle<FQueryUserInfo> QueryUserInfo(FQueryUserInfo::Params&& Params) override;
+	UE_API virtual TOnlineResult<FGetUserInfo> GetUserInfo(FGetUserInfo::Params&& Params) override;
+	UE_API virtual TOnlineAsyncOpHandle<FQueryUserAvatar> QueryUserAvatar(FQueryUserAvatar::Params&& Params) override;
+	UE_API virtual TOnlineResult<FGetUserAvatar> GetUserAvatar(FGetUserAvatar::Params&& Params) override;
+	UE_API virtual TOnlineAsyncOpHandle<FShowUserProfile> ShowUserProfile(FShowUserProfile::Params&& Params) override;
 };
 
 /* UE::Online */ }
+
+#undef UE_API
diff --git a/Engine/Plugins/Online/OnlineServices/Source/OnlineServicesCommonEngineUtils/Public/Online/SessionsLAN.h b/Engine/Plugins/Online/OnlineServices/Source/OnlineServicesCommonEngineUtils/Public/Online/SessionsLAN.h
--- a/Engine/Plugins/Online/OnlineServices/Source/OnlineServicesCommonEngineUtils/Public/Online/SessionsLAN.h
+++ b/Engine/Plugins/Online/OnlineServices/Source/OnlineServicesCommonEngineUtils/Public/Online/SessionsLAN.h
@@ -4,6 +4,8 @@
 #include "Online/SessionsCommon.h"
 #include "Templates/ValueOrError.h"
 
+#define UE_API ONLINESERVICESCOMMONENGINEUTILS_API
+
 class FInternetAddr;
 class FLANSession;
 class FNboSerializeFromBuffer;
@@ -13,26 +15,26 @@ namespace UE::Online {
 
 class FOnlineServicesCommon;
 
-class ONLINESERVICESCOMMONENGINEUTILS_API FOnlineSessionIdRegistryLAN : public FOnlineSessionIdStringRegistry
+class FOnlineSessionIdRegistryLAN : public FOnlineSessionIdStringRegistry
 {
 public:
-	static FOnlineSessionIdRegistryLAN& GetChecked(EOnlineServices ServicesType);
+	static UE_API FOnlineSessionIdRegistryLAN& GetChecked(EOnlineServices ServicesType);
 
-	FOnlineSessionId GetNextSessionId();
+	UE_API FOnlineSessionId GetNextSessionId();
 
 protected:
-	FOnlineSessionIdRegistryLAN(EOnlineServices ServicesType);
+	UE_API FOnlineSessionIdRegistryLAN(EOnlineServices ServicesType);
 };
 
-class ONLINESERVICESCOMMONENGINEUTILS_API FSessionLAN : public FSessionCommon
+class FSessionLAN : public FSessionCommon
 {
 public:
-	FSessionLAN();
+	UE_API FSessionLAN();
 
-	static FSessionLAN& Cast(FSessionCommon& InSession);
-	static const FSessionLAN& Cast(const ISession& InSession);
+	static UE_API FSessionLAN& Cast(FSessionCommon& InSession);
+	static UE_API const FSessionLAN& Cast(const ISession& InSession);
 
-	virtual void DumpState() const override;
+	UE_API virtual void DumpState() const override;
 
 private:
 	void Initialize();
@@ -49,33 +51,33 @@ void ONLINESERVICESCOMMONENGINEUTILS_API SerializeFromBuffer(FNboSerializeFromBu
 
 /* NboSerializerLANSvc */ }
 
-class ONLINESERVICESCOMMONENGINEUTILS_API FSessionsLAN : public FSessionsCommon
+class FSessionsLAN : public FSessionsCommon
 {
 public:
 	using Super = FSessionsCommon;
 
-	FSessionsLAN(FOnlineServicesCommon& InServices);
-	virtual ~FSessionsLAN() override;
+	UE_API FSessionsLAN(FOnlineServicesCommon& InServices);
+	UE_API virtual ~FSessionsLAN() override;
 
-	virtual void Tick(float DeltaSeconds) override;
+	UE_API virtual void Tick(float DeltaSeconds) override;
 
 	// FSessionsCommon
-	virtual TFuture<TOnlineResult<FCreateSession>> CreateSessionImpl(const FCreateSession::Params& Params) override;
-	virtual TFuture<TOnlineResult<FUpdateSessionSettings>> UpdateSessionSettingsImpl(const FUpdateSessionSettings::Params& Params) override;
-	virtual TFuture<TOnlineResult<FLeaveSession>> LeaveSessionImpl(const FLeaveSession::Params& Params) override;
-	virtual TFuture<TOnlineResult<FFindSessions>> FindSessionsImpl(const FFindSessions::Params& Params) override;
-	virtual TFuture<TOnlineResult<FJoinSession>> JoinSessionImpl(const FJoinSession::Params& Params) override;
+	UE_API virtual TFuture<TOnlineResult<FCreateSession>> CreateSessionImpl(const FCreateSession::Params& Params) override;
+	UE_API virtual TFuture<TOnlineResult<FUpdateSessionSettings>> UpdateSessionSettingsImpl(const FUpdateSessionSettings::Params& Params) override;
+	UE_API virtual TFuture<TOnlineResult<FLeaveSession>> LeaveSessionImpl(const FLeaveSession::Params& Params) override;
+	UE_API virtual TFuture<TOnlineResult<FFindSessions>> FindSessionsImpl(const FFindSessions::Params& Params) override;
+	UE_API virtual TFuture<TOnlineResult<FJoinSession>> JoinSessionImpl(const FJoinSession::Params& Params) override;
 
 protected:
 	using FHostSessionResult = TValueOrError<void, FOnlineError>;
 
-	FHostSessionResult TryHostLANSession();
-	TSharedRef<FSessionLAN> AddLANSession(const FCreateSession::Params& Params);
-	void FindLANSessions(const FAccountId& LocalAccountId);
-	void StopLANSession();
-	void OnValidQueryPacketReceived(uint8* PacketData, int32 PacketLength, uint64 ClientNonce);
-	void OnValidResponsePacketReceived(uint8* PacketData, int32 PacketLength, const FAccountId LocalAccountId);
-	void OnLANSearchTimeout(const FAccountId LocalAccountId);
+	UE_API FHostSessionResult TryHostLANSession();
+	UE_API TSharedRef<FSessionLAN> AddLANSession(const FCreateSession::Params& Params);
+	UE_API void FindLANSessions(const FAccountId& LocalAccountId);
+	UE_API void StopLANSession();
+	UE_API void OnValidQueryPacketReceived(uint8* PacketData, int32 PacketLength, uint64 ClientNonce);
+	UE_API void OnValidResponsePacketReceived(uint8* PacketData, int32 PacketLength, const FAccountId LocalAccountId);
+	UE_API void OnLANSearchTimeout(const FAccountId LocalAccountId);
 	virtual void AppendSessionToPacket(FNboSerializeToBuffer& Packet, const FSessionLAN& Session) = 0;
 	virtual void ReadSessionFromPacket(FNboSerializeFromBuffer& Packet, FSessionLAN& Session) = 0;
 
@@ -84,3 +86,5 @@ protected:
 };
 
 /* UE::Online */ }
+
+#undef UE_API
diff --git a/Engine/Plugins/Online/OnlineServices/Source/OnlineServicesInterface/Public/Online/OnlineServices.h b/Engine/Plugins/Online/OnlineServices/Source/OnlineServicesInterface/Public/Online/OnlineServices.h
--- a/Engine/Plugins/Online/OnlineServices/Source/OnlineServicesInterface/Public/Online/OnlineServices.h
+++ b/Engine/Plugins/Online/OnlineServices/Source/OnlineServicesInterface/Public/Online/OnlineServices.h
@@ -49,7 +49,7 @@ struct FGetResolvedConnectString
 	};
 };
 
-class ONLINESERVICESINTERFACE_API IOnlineServices
+class IOnlineServices
 {
 public:
 	/**
diff --git a/Engine/Plugins/Online/OnlineServices/Source/OnlineServicesInterface/Public/Online/OnlineServicesRegistry.h b/Engine/Plugins/Online/OnlineServices/Source/OnlineServicesInterface/Public/Online/OnlineServicesRegistry.h
--- a/Engine/Plugins/Online/OnlineServices/Source/OnlineServicesInterface/Public/Online/OnlineServicesRegistry.h
+++ b/Engine/Plugins/Online/OnlineServices/Source/OnlineServicesInterface/Public/Online/OnlineServicesRegistry.h
@@ -13,7 +13,7 @@ namespace UE::Online {
 	
 enum class EOnlineServices : uint8;
 
-class ONLINESERVICESINTERFACE_API IOnlineServicesFactory
+class IOnlineServicesFactory
 {
 public:
 	/**
diff --git a/Engine/Plugins/Online/OnlineServices/Source/OnlineServicesInterface/Public/Online/Schema.h b/Engine/Plugins/Online/OnlineServices/Source/OnlineServicesInterface/Public/Online/Schema.h
--- a/Engine/Plugins/Online/OnlineServices/Source/OnlineServicesInterface/Public/Online/Schema.h
+++ b/Engine/Plugins/Online/OnlineServices/Source/OnlineServicesInterface/Public/Online/Schema.h
@@ -5,6 +5,8 @@
 #include "Online/SchemaTypes.h"
 #include "Templates/SharedPointer.h"
 
+#define UE_API ONLINESERVICESINTERFACE_API
+
 namespace UE::Online { template <typename OpType> class TOnlineResult; }
 
 namespace UE::Online {
@@ -13,7 +15,7 @@ class FSchemaRegistry;
 
 namespace Private {
 
-class ONLINESERVICESINTERFACE_API FSchemaCategoryInstanceBase
+class FSchemaCategoryInstanceBase
 {
 public:
 	// DerivedSchemaId may remain unset in two situations:
@@ -21,7 +23,7 @@ public:
 	//    the base schema is required to have an attribute flagged as SchemaCompatibilityId so that
 	//    the derived schema can be discovered.
 	// 2. Schema swapping is not enabled. All attributes are defined in the base schema.
-	FSchemaCategoryInstanceBase(
+	UE_API FSchemaCategoryInstanceBase(
 		const FSchemaId& DerivedSchemaId,
 		const FSchemaId& BaseSchemaId,
 		const FSchemaCategoryId& CategoryId,
@@ -31,15 +33,15 @@ public:
 
 	// Two phase commit to the service.
 	// Prepare data to be written to the service from an client delta of attributes.
-	TOnlineResult<FSchemaCategoryInstancePrepareClientChanges> PrepareClientChanges(FSchemaCategoryInstancePrepareClientChanges::Params&& Params) const;
+	UE_API TOnlineResult<FSchemaCategoryInstancePrepareClientChanges> PrepareClientChanges(FSchemaCategoryInstancePrepareClientChanges::Params&& Params) const;
 	// Commit the client data once the service data has been successfully written.
-	FSchemaCategoryInstanceCommitClientChanges::Result CommitClientChanges();
+	UE_API FSchemaCategoryInstanceCommitClientChanges::Result CommitClientChanges();
 
 	// Two phase commit from the service.
 	// Prepare data to be written to the client from a service snapshot of attributes.
-	TOnlineResult<FSchemaCategoryInstancePrepareServiceSnapshot> PrepareServiceSnapshot(FSchemaCategoryInstancePrepareServiceSnapshot::Params&& Params) const;
+	UE_API TOnlineResult<FSchemaCategoryInstancePrepareServiceSnapshot> PrepareServiceSnapshot(FSchemaCategoryInstancePrepareServiceSnapshot::Params&& Params) const;
 	// Commit a service snapshot containing all known service attributes.
-	FSchemaCategoryInstanceCommitServiceSnapshot::Result CommitServiceSnapshot();
+	UE_API FSchemaCategoryInstanceCommitServiceSnapshot::Result CommitServiceSnapshot();
 
 	// Todo: ApplyServiceDelta
 
@@ -47,47 +49,47 @@ public:
 	TSharedPtr<const FSchemaDefinition> GetBaseDefinition() const { return BaseSchemaDefinition; };
 
 	// Check whether the schema is valid.
-	bool IsValid() const;
+	UE_API bool IsValid() const;
 
 	// Check whether an client attribute is valid in the base schema.
 	// Intended to be used when attribute information needs to be verified without translating to the service.
-	bool VerifyBaseAttributeData(
+	UE_API bool VerifyBaseAttributeData(
 		const FSchemaAttributeId& Id,
 		const FSchemaVariant& Data,
 		FSchemaServiceAttributeId& OutSchemaServiceAttributeId,
 		ESchemaServiceAttributeFlags& OutSchemaServiceAttributeFlags);
 
 protected:
-	const TMap<FSchemaAttributeId, FSchemaVariant>& GetClientSnapshot() const;
+	UE_API const TMap<FSchemaAttributeId, FSchemaVariant>& GetClientSnapshot() const;
 	virtual TMap<FSchemaAttributeId, FSchemaVariant>& GetMutableClientSnapshot() = 0;
 
 private:
-	bool InitializeSchemaDefinition(
+	UE_API bool InitializeSchemaDefinition(
 		int64 SchemaCompatibilityId,
 		const FSchemaDefinition* OptionalBaseDefinition,
 		const FSchemaCategoryId& CategoryId,
 		TSharedPtr<const FSchemaDefinition>* OutSchemaDefinition,
 		const FSchemaCategoryDefinition** OutSchemaCategoryDefinition) const;
 
-	bool InitializeSchemaDefinition(
+	UE_API bool InitializeSchemaDefinition(
 		const FSchemaId& SchemaId,
 		const FSchemaDefinition* OptionalBaseDefinition,
 		const FSchemaCategoryId& CategoryId,
 		TSharedPtr<const FSchemaDefinition>* OutSchemaDefinition,
 		const FSchemaCategoryDefinition** OutSchemaCategoryDefinition) const;
 
-	bool InitializeSchemaDefinition(
+	UE_API bool InitializeSchemaDefinition(
 		const TSharedRef<const FSchemaDefinition>& NewDefinition,
 		const FSchemaDefinition* OptionalBaseDefinition,
 		const FSchemaCategoryId& CategoryId,
 		TSharedPtr<const FSchemaDefinition>* OutSchemaDefinition,
 		const FSchemaCategoryDefinition** OutSchemaCategoryDefinition) const;
 
-	bool GetSerializationSchema(
+	UE_API bool GetSerializationSchema(
 		const FSchemaDefinition** OutSchemaDefinition,
 		const FSchemaCategoryDefinition** OutSchemaCategoryDefinition) const;
 
-	void ResetPreparedChanges() const;
+	UE_API void ResetPreparedChanges() const;
 
 	struct FPreparedClientChanges
 	{
@@ -165,15 +167,15 @@ private:
 /** Default implementation with both translation and client attribute snapshot data contained within the category instance. */
 using FSchemaCategoryInstance = TSchemaCategoryInstance<FSchemaCategoryInstanceInternalSnapshotAccessor>;
 
-class ONLINESERVICESINTERFACE_API FSchemaRegistry
+class FSchemaRegistry
 {
 public:
 	// Parse the loaded config structures.
-	bool ParseConfig(const FSchemaRegistryDescriptorConfig& Config);
+	UE_API bool ParseConfig(const FSchemaRegistryDescriptorConfig& Config);
 
-	TSharedPtr<const FSchemaDefinition> GetDefinition(const FSchemaId& SchemaId) const;
-	TSharedPtr<const FSchemaDefinition> GetDefinition(int64 CompatibilityId) const;
-	bool IsSchemaChildOf(const FSchemaId& SchemaId, const FSchemaId& ParentSchemaId) const;
+	UE_API TSharedPtr<const FSchemaDefinition> GetDefinition(const FSchemaId& SchemaId) const;
+	UE_API TSharedPtr<const FSchemaDefinition> GetDefinition(int64 CompatibilityId) const;
+	UE_API bool IsSchemaChildOf(const FSchemaId& SchemaId, const FSchemaId& ParentSchemaId) const;
 
 private:
 
@@ -182,3 +184,5 @@ private:
 };
 
 /* UE::Online */ }
+
+#undef UE_API
diff --git a/Engine/Plugins/Online/OnlineServices/Source/OnlineServicesInterface/Public/Online/SchemaTypes.h b/Engine/Plugins/Online/OnlineServices/Source/OnlineServicesInterface/Public/Online/SchemaTypes.h
--- a/Engine/Plugins/Online/OnlineServices/Source/OnlineServicesInterface/Public/Online/SchemaTypes.h
+++ b/Engine/Plugins/Online/OnlineServices/Source/OnlineServicesInterface/Public/Online/SchemaTypes.h
@@ -8,6 +8,8 @@
 #include "Online/OnlineMeta.h"
 #include "UObject/NameTypes.h"
 
+#define UE_API ONLINESERVICESINTERFACE_API
+
 namespace UE::Online::Meta { template <typename StructType> struct TStructDetails; }
 
 namespace UE::Online {
@@ -191,14 +193,14 @@ ENUM_CLASS_FLAGS(ESchemaServiceAttributeSupportedTypeFlags);
 ONLINESERVICESINTERFACE_API const TCHAR* LexToString(ESchemaServiceAttributeSupportedTypeFlags SchemaServiceAttributeSupportedTypeFlags);
 ONLINESERVICESINTERFACE_API void LexFromString(ESchemaServiceAttributeSupportedTypeFlags& OutSchemaServiceAttributeSupportedTypeFlags, const TCHAR* InStr);
 
-class ONLINESERVICESINTERFACE_API FSchemaVariant final
+class FSchemaVariant final
 {
 public:
 	using FVariantType = TVariant<FString, int64, double, bool>;
 	FSchemaVariant() = default;
 	FSchemaVariant(const FSchemaVariant& InOther) = default;
-	FSchemaVariant(FSchemaVariant&& InOther);
-	FSchemaVariant& operator=(FSchemaVariant&&);
+	UE_API FSchemaVariant(FSchemaVariant&& InOther);
+	UE_API FSchemaVariant& operator=(FSchemaVariant&&);
 	FSchemaVariant& operator=(const FSchemaVariant&) = default;
 
 	FSchemaVariant(const TCHAR* InData) { Set(InData); }
@@ -207,23 +209,23 @@ public:
 	FSchemaVariant(int64 InData) { Set(InData); }
 	FSchemaVariant(double InData) { Set(InData); }
 	FSchemaVariant(bool InData) { Set(InData); }
-	void Set(const TCHAR* AsString);
-	void Set(const FString& AsString);
-	void Set(FString&& AsString);
-	void Set(int64 AsInt);
-	void Set(double AsDouble);
-	void Set(bool bAsBool);
+	UE_API void Set(const TCHAR* AsString);
+	UE_API void Set(const FString& AsString);
+	UE_API void Set(FString&& AsString);
+	UE_API void Set(int64 AsInt);
+	UE_API void Set(double AsDouble);
+	UE_API void Set(bool bAsBool);
 
 	ESchemaAttributeType GetType() const { return VariantType; }
 
-	int64 GetInt64() const;
-	double GetDouble() const;
-	bool GetBoolean() const;
-	const FString& GetString() const;
+	UE_API int64 GetInt64() const;
+	UE_API double GetDouble() const;
+	UE_API bool GetBoolean() const;
+	UE_API const FString& GetString() const;
 
-	FString ToLogString() const;
+	UE_API FString ToLogString() const;
 
-	bool operator==(const FSchemaVariant& Other) const;
+	UE_API bool operator==(const FSchemaVariant& Other) const;
 	bool operator!=(const FSchemaVariant& Other) const { return !(*this == Other); }
 public:
 	FVariantType VariantData;
@@ -718,3 +720,5 @@ END_ONLINE_STRUCT_META()
 /* Meta*/ }
 
 /* UE::Online */ }
+
+#undef UE_API
diff --git a/Engine/Plugins/Online/OnlineServices/Source/OnlineServicesInterface/Public/Online/Sessions.h b/Engine/Plugins/Online/OnlineServices/Source/OnlineServicesInterface/Public/Online/Sessions.h
--- a/Engine/Plugins/Online/OnlineServices/Source/OnlineServicesInterface/Public/Online/Sessions.h
+++ b/Engine/Plugins/Online/OnlineServices/Source/OnlineServicesInterface/Public/Online/Sessions.h
@@ -7,6 +7,8 @@
 #include "Online/SchemaTypes.h"
 #include "Online/OnlineResult.h"
 
+#define UE_API ONLINESERVICESINTERFACE_API
+
 namespace UE::Online { template <typename DelegateSignature> class TOnlineEvent; }
 namespace UE::Online { template <typename OpType> class TOnlineAsyncOpHandle; }
 
@@ -72,7 +74,7 @@ ONLINESERVICESINTERFACE_API const TCHAR* LexToString(ESessionJoinPolicy Value);
 ONLINESERVICESINTERFACE_API void LexFromString(ESessionJoinPolicy& Value, const TCHAR* InStr);
 
 /** Contains new values for an FSessions modifiable settings. Taken as a parameter by FUpdateSessionSettings method */
-struct ONLINESERVICESINTERFACE_API FSessionSettingsUpdate
+struct FSessionSettingsUpdate
 {
 	/** Set with an updated value if the SchemaName field will be changed in the update operation */
 	TOptional<FSchemaId> SchemaName;
@@ -88,7 +90,7 @@ struct ONLINESERVICESINTERFACE_API FSessionSettingsUpdate
 	/** Names of custom settings to be removed in the update operation */
 	TArray<FSchemaAttributeId> RemovedCustomSettings;
 
-	FSessionSettingsUpdate& operator+=(FSessionSettingsUpdate&& UpdatedValue);
+	UE_API FSessionSettingsUpdate& operator+=(FSessionSettingsUpdate&& UpdatedValue);
 };
 
 /** Contains updated data for any modifiable members of FSessionSettings. Part of the FSessionUpdated event data */
@@ -114,7 +116,7 @@ struct FSessionSettingsChanges
 };
 
 /** Set of all of an FSession's defining properties that can be updated by the session owner during its lifetime, using the FUpdateSessionSettings method */
-struct ONLINESERVICESINTERFACE_API FSessionSettings
+struct FSessionSettings
 {
 	/* The name for the schema which will be applied to the session's user-defined attributes */
 	FName SchemaName;
@@ -131,7 +133,7 @@ struct ONLINESERVICESINTERFACE_API FSessionSettings
 	/* Map of user-defined settings to be passed to the platform APIs as additional information for various purposes */
 	FCustomSessionSettingsMap CustomSettings;
 
-	FSessionSettings& operator+=(const FSessionSettingsChanges& UpdatedValue);
+	UE_API FSessionSettings& operator+=(const FSessionSettingsChanges& UpdatedValue);
 };
 
 /** Information about an FSession that will be set at creation time and remain constant during its lifetime */
@@ -626,7 +628,7 @@ struct FSessionLeft
 };
 
 /** Contains updated data for any modifiable members of ISession */
-struct ONLINESERVICESINTERFACE_API FSessionUpdate
+struct FSessionUpdate
 {
 	/** If set, the OwnerUserId member will have updated to this value */
 	TOptional<FAccountId> OwnerAccountId;
@@ -640,7 +642,7 @@ struct ONLINESERVICESINTERFACE_API FSessionUpdate
 	/** Id handles for members that just left the session */
 	FSessionMemberIdsSet RemovedSessionMembers;
 
-	FSessionUpdate& operator+=(const FSessionUpdate& SessionUpdate);
+	UE_API FSessionUpdate& operator+=(const FSessionUpdate& SessionUpdate);
 };
 
 struct FSessionUpdated
@@ -1133,3 +1135,5 @@ END_ONLINE_STRUCT_META()
 /* Meta*/ }
 
 /* UE::Online */ }
+
+#undef UE_API
