Fix the memory leaks related to circular reference of shared pointer in OSSv2 presence.

diff --git a/Engine/Plugins/Online/OnlineServices/Source/OnlineServicesCommon/Private/Online/PresenceCommon.cpp b/Engine/Plugins/Online/OnlineServices/Source/OnlineServicesCommon/Private/Online/PresenceCommon.cpp
--- a/Engine/Plugins/Online/OnlineServices/Source/OnlineServicesCommon/Private/Online/PresenceCommon.cpp
+++ b/Engine/Plugins/Online/OnlineServices/Source/OnlineServicesCommon/Private/Online/PresenceCommon.cpp
@@ -39,51 +39,56 @@ public:
 class FBatchQueryPresenceHelper
 {
 public:
-	FBatchQueryPresenceHelper(TOnlineAsyncOpRef<FBatchQueryPresence> InOp, TSharedPtr<IPresence> InPresence)
+	FBatchQueryPresenceHelper(TOnlineAsyncOpPtr<FBatchQueryPresence> InOp, TSharedPtr<IPresence> InPresence)
 		: Op(InOp)
 		, Presence(InPresence)
 		, ErrorState(Errors::Success())
 	{
-		NumPresencesRemaining = Op->GetParams().TargetAccountIds.Num();
+		NumPresencesRemaining = InOp->GetParams().TargetAccountIds.Num();
 	}
 
 	TFuture<FBatchQueryPresenceHelperResult> GetPromise()
 	{
-		for(FAccountId TargetAccountId : Op->GetParams().TargetAccountIds)
+		TOnlineAsyncOpPtr<FBatchQueryPresence> OpPtr = Op.Pin();
+		if (OpPtr.IsValid())
 		{
-			FQueryPresence::Params Params;
-			Params.LocalAccountId = Op->GetParams().LocalAccountId;
-			Params.TargetAccountId = TargetAccountId;
-			Params.bListenToChanges = Op->GetParams().bListenToChanges;
-
-			Presence->QueryPresence(MoveTemp(Params)).OnComplete([this](const TOnlineResult<FQueryPresence> NewOp) mutable
+			for(FAccountId TargetAccountId : OpPtr->GetParams().TargetAccountIds)
 			{
-				if (NewOp.IsOk())
-				{
-					// push results
-					Result.Presences.Add(NewOp.GetOkValue().Presence);
-				}
-				else
-				{
-					// error
-					// we're only reporting the most recent one, more verbose logging on each individual error can be found in the logs
-					ErrorState = NewOp.GetErrorValue();
-				}
+				FQueryPresence::Params Params;
+				Params.LocalAccountId = OpPtr->GetParams().LocalAccountId;
+				Params.TargetAccountId = TargetAccountId;
+				Params.bListenToChanges = OpPtr->GetParams().bListenToChanges;
 
-
-				// final calculations
-				if (--NumPresencesRemaining <= 0)
+				Presence->QueryPresence(MoveTemp(Params)).OnComplete([this](const TOnlineResult<FQueryPresence> NewOp) mutable
 				{
-					Promise.SetValue(FBatchQueryPresenceHelperResult{ MoveTemp(Result), MoveTemp(ErrorState) });
-				}
-			});
+					if (NewOp.IsOk())
+					{
+						// push results
+						Result.Presences.Add(NewOp.GetOkValue().Presence);
+					}
+					else
+					{
+						// error
+						// we're only reporting the most recent one, more verbose logging on each individual error can be found in the logs
+						ErrorState = NewOp.GetErrorValue();
+					}
+
+
+					// final calculations
+					if (--NumPresencesRemaining <= 0)
+					{
+						Promise.SetValue(FBatchQueryPresenceHelperResult{ MoveTemp(Result), MoveTemp(ErrorState) });
+					}
+				});
+			}
 		}
 
 		return Promise.GetFuture();
 	}
 
 	FBatchQueryPresence::Result Result;
-	TOnlineAsyncOpRef<FBatchQueryPresence> Op;
+	// Need to use weak ptr here otherwise there will be circular reference with Op->Steps->Lambda->Helper->Op
+	TWeakPtr<TOnlineAsyncOp<FBatchQueryPresence>> Op;
 	TSharedPtr<IPresence> Presence;
 	FOnlineError ErrorState;
 	int NumPresencesRemaining = 0;
@@ -94,7 +99,7 @@ public:
 TOnlineAsyncOpHandle<FBatchQueryPresence> FPresenceCommon::BatchQueryPresence(FBatchQueryPresence::Params&& Params)
 {
 	TOnlineAsyncOpRef<FBatchQueryPresence> Op = GetOp<FBatchQueryPresence>(MoveTemp(Params));
-	TSharedRef<FBatchQueryPresenceHelper> Helper = MakeShared<FBatchQueryPresenceHelper>(Op, this->AsShared());
+	TSharedRef<FBatchQueryPresenceHelper> Helper = MakeShared<FBatchQueryPresenceHelper>(Op.ToSharedPtr(), this->AsShared());
 
 	Op->Then([this, Helper](TOnlineAsyncOp<FBatchQueryPresence>& InAsyncOp) mutable
 	{
