Add explicit Clear method to FWorldContextScopedObjectCache to enable control for low level tests

diff --git a/Engine/Plugins/Online/OnlineServices/Source/OnlineServicesCommonEngineUtils/Public/Online/WorldContextScopedObjectCache.h b/Engine/Plugins/Online/OnlineServices/Source/OnlineServicesCommonEngineUtils/Public/Online/WorldContextScopedObjectCache.h
--- a/Engine/Plugins/Online/OnlineServices/Source/OnlineServicesCommonEngineUtils/Public/Online/WorldContextScopedObjectCache.h
+++ b/Engine/Plugins/Online/OnlineServices/Source/OnlineServicesCommonEngineUtils/Public/Online/WorldContextScopedObjectCache.h
@@ -65,6 +65,12 @@ public:
 #endif // UE_EDITOR && !UE_IS_COOKED_EDITOR
 	}
 
+	/** Explicitly clear an object for the given context name. Expected to only be used in strictly controlled scenarios, such as automated tests. */
+	void Clear(FName ContextName)
+	{
+		Objects.Remove(ContextName);
+	}
+
 private:
 	TMap<FName, TSharedRef<ObjectType>> Objects;
 #if UE_EDITOR && !UE_IS_COOKED_EDITOR
