Adding PS platform base classes for UserInfo

diff --git a/Engine/Plugins/Online/OnlineServices/Source/OnlineServicesCommon/Public/Online/OnlineAsyncOp.h b/Engine/Plugins/Online/OnlineServices/Source/OnlineServicesCommon/Public/Online/OnlineAsyncOp.h
--- a/Engine/Plugins/Online/OnlineServices/Source/OnlineServicesCommon/Public/Online/OnlineAsyncOp.h
+++ b/Engine/Plugins/Online/OnlineServices/Source/OnlineServicesCommon/Public/Online/OnlineAsyncOp.h
@@ -2,14 +2,15 @@
 
 #pragma once
 
+#include "Async/Async.h"
+#include "Containers/Ticker.h"
 #include "Online/OnlineAsyncOpHandle.h"
 #include "Online/OnlineResult.h"
 #include "Online/OnlineServicesDelegates.h"
 #include "Online/OnlineServicesLog.h"
 #include "Online/OnlineTypeInfo.h"
-#include "Async/Async.h"
-#include "Containers/Ticker.h"
 #include "Traits/ElementType.h"
+#include "Traits/IsContiguousContainer.h"
 
 namespace UE::Online {
 
@@ -390,30 +391,28 @@ public:
 					TPromise<ResultType> Promise;
 					// set promise continuation before calling the callable so that we will complete the step as soon as the value is set
 					Promise.GetFuture()
-						.Next([this, WeakOperation](const ResultType& Value)
+					.Next([this, WeakOperation](const ResultType& Value)
+					{
+						if (TSharedPtr<TOnlineAsyncOp<OpType>> PinnedOperation2 = WeakOperation.Pin())
 						{
-							TSharedPtr<TOnlineAsyncOp<OpType>> PinnedOperation2 = WeakOperation.Pin();
-							if (PinnedOperation2)
-							{
-								EmplaceResult(Value);
-								PinnedOperation2->ExecuteNextStep();
-							}
-						});
+							EmplaceResult(Value);
+							PinnedOperation2->ExecuteNextStep();
+						}
+					});
 
 					Callable(*PinnedOperation, MoveTempIfPossible(LastResult), MoveTemp(Promise));
 				}
 				else if constexpr (TOnlineAsyncOpCallableTraits<CallableType>::bAsyncResult)
 				{
 					Callable(*PinnedOperation, MoveTempIfPossible(LastResult))
-						.Next([this, WeakOperation](const ResultType& Value)
+					.Next([this, WeakOperation](const ResultType& Value)
+					{
+						if (TSharedPtr<TOnlineAsyncOp<OpType>> PinnedOperation2 = WeakOperation.Pin())
 						{
-							TSharedPtr<TOnlineAsyncOp<OpType>> PinnedOperation2 = WeakOperation.Pin();
-							if (PinnedOperation2)
-							{
-								EmplaceResult(Value);
-								PinnedOperation2->ExecuteNextStep();
-							}
-						});
+							EmplaceResult(Value);
+							PinnedOperation2->ExecuteNextStep();
+						}
+					});
 				}
 				else
 				{
@@ -437,30 +436,28 @@ public:
 					TPromise<ResultType> Promise;
 					// set promise continuation before calling the callable so that we will complete the step as soon as the value is set
 					Promise.GetFuture()
-						.Next([this, WeakOperation](const ResultType& Value)
+					.Next([this, WeakOperation](const ResultType& Value)
+					{
+						if (TSharedPtr<TOnlineAsyncOp<OpType>> PinnedOperation2 = WeakOperation.Pin())
 						{
-							TSharedPtr<TOnlineAsyncOp<OpType>> PinnedOperation2 = WeakOperation.Pin();
-							if (PinnedOperation2)
-							{
-								EmplaceResult(Value);
-								PinnedOperation2->ExecuteNextStep();
-							}
-						});
+							EmplaceResult(Value);
+							PinnedOperation2->ExecuteNextStep();
+						}
+					});
 
 					Callable(*PinnedOperation, MoveTemp(Promise));
 				}
 				else if constexpr (TOnlineAsyncOpCallableTraits<CallableType>::bAsyncResult)
 				{
 					Callable(*PinnedOperation)
-						.Next([this, WeakOperation](const ResultType& Value)
+					.Next([this, WeakOperation](const ResultType& Value)
+					{
+						if (TSharedPtr<TOnlineAsyncOp<OpType>> PinnedOperation2 = WeakOperation.Pin())
 						{
-							TSharedPtr<TOnlineAsyncOp<OpType>> PinnedOperation2 = WeakOperation.Pin();
-							if (PinnedOperation2)
-							{
-								EmplaceResult(Value);
-								PinnedOperation2->ExecuteNextStep();
-							}
-						});
+							EmplaceResult(Value);
+							PinnedOperation2->ExecuteNextStep();
+						}
+					});
 				}
 				else
 				{
@@ -524,14 +521,13 @@ public:
 				if constexpr (TOnlineAsyncOpCallableTraits<CallableType>::bAsyncResult)
 				{
 					Callable(*PinnedOperation, MoveTempIfPossible(LastResult))
-						.Next([this, WeakOperation]()
+					.Next([this, WeakOperation]()
+					{
+						if (TSharedPtr<TOnlineAsyncOp<OpType>> PinnedOperation2 = WeakOperation.Pin())
 						{
-							TSharedPtr<TOnlineAsyncOp<OpType>> PinnedOperation2 = WeakOperation.Pin();
-							if (PinnedOperation2)
-							{
-								PinnedOperation2->ExecuteNextStep();
-							}
-						});
+							PinnedOperation2->ExecuteNextStep();
+						}
+					});
 				}
 				else
 				{
@@ -553,14 +549,13 @@ public:
 				if constexpr (TOnlineAsyncOpCallableTraits<CallableType>::bAsyncResult)
 				{
 					Callable(*PinnedOperation)
-						.Next([this, WeakOperation]()
+					.Next([this, WeakOperation]()
+					{
+						if (TSharedPtr<TOnlineAsyncOp<OpType>> PinnedOperation2 = WeakOperation.Pin())
 						{
-							TSharedPtr<TOnlineAsyncOp<OpType>> PinnedOperation2 = WeakOperation.Pin();
-							if (PinnedOperation2)
-							{
-								PinnedOperation2->ExecuteNextStep();
-							}
-						});
+							PinnedOperation2->ExecuteNextStep();
+						}
+					});
 				}
 				else
 				{
@@ -617,46 +612,41 @@ class TStep<TContinuationResult<ResultType>> : public IStep
 						TPromise<TContinuationResult<ResultType>> Promise;
 						// set promise continuation before calling the callable so that we will complete the step as soon as the value is set
 						Promise.GetFuture()
-							.Next([this, WeakOperation](const TContinuationResult<ResultType>& Value)
+						.Next([this, WeakOperation](const TContinuationResult<ResultType>& Value)
+						{
+							if (TSharedPtr<TOnlineAsyncOp<OpType>> PinnedOperation2 = WeakOperation.Pin())
+							{
+								if (!Value.IsComplete())
 								{
-									TSharedPtr<TOnlineAsyncOp<OpType>> PinnedOperation2 = WeakOperation.Pin();
-
-									if (PinnedOperation2)
-									{
-										if (!Value.IsComplete())
-										{
-											PinnedOperation2->ExecuteRepeatStep();
-										}
-										else
-										{
-											EmplaceResult(Value);
-											PinnedOperation2->ExecuteNextStep();
-										}
-									}
-								});
+									PinnedOperation2->ExecuteRepeatStep();
+								}
+								else
+								{
+									EmplaceResult(Value);
+									PinnedOperation2->ExecuteNextStep();
+								}
+							}
+						});
 						Callable(*PinnedOperation, MoveTempIfPossible(LastResult), MoveTemp(Promise));
 					}
 					else if constexpr (TOnlineAsyncOpCallableTraits<CallableType>::bAsyncResult)
 					{
 						Callable(*PinnedOperation, LastResult)
-							.Next([this, WeakOperation](const TContinuationResult<ResultType>& Value)
+						.Next([this, WeakOperation](const TContinuationResult<ResultType>& Value)
+						{
+							if (TSharedPtr<TOnlineAsyncOp<OpType>> PinnedOperation2 = WeakOperation.Pin())
+							{
+								if (!Value.IsComplete())
 								{
-
-									TSharedPtr<TOnlineAsyncOp<OpType>> PinnedOperation2 = WeakOperation.Pin();
-									if (PinnedOperation2)
-									{
-										if (!Value.IsComplete())
-										{
-											PinnedOperation2->ExecuteRepeatStep();
-										}
-										else
-										{
-											EmplaceResult(Value.GetResult());
-											PinnedOperation2->ExecuteNextStep();
-										}
-									}
-								});
-
+									PinnedOperation2->ExecuteRepeatStep();
+								}
+								else
+								{
+									EmplaceResult(Value.GetResult());
+									PinnedOperation2->ExecuteNextStep();
+								}
+							}
+						});
 					}
 					else
 					{
@@ -688,47 +678,41 @@ class TStep<TContinuationResult<ResultType>> : public IStep
 						TPromise<TContinuationResult<ResultType>> Promise;
 						// set promise continuation before calling the callable so that we will complete the step as soon as the value is set
 						Promise.GetFuture()
-							.Next([this, WeakOperation](const TContinuationResult<ResultType>& Value)
+						.Next([this, WeakOperation](const TContinuationResult<ResultType>& Value)
+						{
+							if (TSharedPtr<TOnlineAsyncOp<OpType>> PinnedOperation2 = WeakOperation.Pin())
+							{
+								if (!Value.IsComplete())
 								{
-									TSharedPtr<TOnlineAsyncOp<OpType>> PinnedOperation2 = WeakOperation.Pin();
-									if (PinnedOperation2)
-									{
-										if (!Value.IsComplete())
-										{
-											PinnedOperation2->ExecuteRepeatStep();
-										}
-										else
-										{
-											EmplaceResult(Value);
-											PinnedOperation2->ExecuteNextStep();
-										}
-									}
-								});
+									PinnedOperation2->ExecuteRepeatStep();
+								}
+								else
+								{
+									EmplaceResult(Value);
+									PinnedOperation2->ExecuteNextStep();
+								}
+							}
+						});
 						Callable(*PinnedOperation, MoveTemp(Promise));
 					}
 					else if constexpr (TOnlineAsyncOpCallableTraits<CallableType>::bAsyncResult)
 					{
-
 						Callable(*PinnedOperation)
-							.Next([this, WeakOperation](const TContinuationResult<ResultType>& Value)
+						.Next([this, WeakOperation](const TContinuationResult<ResultType>& Value)
+						{
+							if (TSharedPtr<TOnlineAsyncOp<OpType>> PinnedOperation2 = WeakOperation.Pin())
+							{
+								if (!Value.IsComplete())
 								{
-
-									TSharedPtr<TOnlineAsyncOp<OpType>> PinnedOperation2 = WeakOperation.Pin();
-									if (PinnedOperation2)
-									{
-										if (!Value.IsComplete())
-										{
-											PinnedOperation2->ExecuteRepeatStep();
-										}
-										else
-										{
-											EmplaceResult(Value);
-											PinnedOperation2->ExecuteNextStep();
-										}
-									}
-
-								});
-
+									PinnedOperation2->ExecuteRepeatStep();
+								}
+								else
+								{
+									EmplaceResult(Value);
+									PinnedOperation2->ExecuteNextStep();
+								}
+							}
+						});
 					}
 					else
 					{
@@ -801,23 +785,20 @@ public:
 					if constexpr (TOnlineAsyncOpCallableTraits<CallableType>::bAsyncResult)
 					{
 						Callable(*PinnedOperation, LastResult)
-							.Next([this, WeakOperation](const TContinuationResult<void>& Value)
+						.Next([this, WeakOperation](const TContinuationResult<void>& Value)
+						{
+							if (TSharedPtr<TOnlineAsyncOp<OpType>> PinnedOperation2 = WeakOperation.Pin())
 							{
-
-								TSharedPtr<TOnlineAsyncOp<OpType>> PinnedOperation2 = WeakOperation.Pin();
-								if (PinnedOperation2)
+								if (!Value.IsComplete())
 								{
-									if (!Value.IsComplete())
-									{
-										PinnedOperation2->ExecuteRepeatStep();
-									}
-									else
-									{
-										PinnedOperation2->ExecuteNextStep();
-									}
+									PinnedOperation2->ExecuteRepeatStep();
 								}
-							});
-
+								else
+								{
+									PinnedOperation2->ExecuteNextStep();
+								}
+							}
+						});
 					}
 					else
 					{
@@ -846,23 +827,21 @@ public:
 					if constexpr (TOnlineAsyncOpCallableTraits<CallableType>::bAsyncResult)
 					{
 						Callable(*PinnedOperation)
-							.Next([this, WeakOperation](const TContinuationResult<void>& Value)
+						.Next([this, WeakOperation](const TContinuationResult<void>& Value)
+						{
+							if (TSharedPtr<TOnlineAsyncOp<OpType>> PinnedOperation2 = WeakOperation.Pin())
 							{
-
-								TSharedPtr<TOnlineAsyncOp<OpType>> PinnedOperation2 = WeakOperation.Pin();
-								if (PinnedOperation2)
+								if (!Value.IsComplete())
 								{
-									if (!Value.IsComplete())
-									{
-										PinnedOperation2->ExecuteRepeatStep();
-									}
-									else
-									{
-										PinnedOperation2->ExecuteNextStep();
-									}
+									PinnedOperation2->ExecuteRepeatStep();
 								}
+								else
+								{
+									PinnedOperation2->ExecuteNextStep();
+								}
+							}
 
-							});
+						});
 					}
 					else
 					{
@@ -939,7 +918,7 @@ protected:
 	FOnlineAsyncExecutionPolicy ExecutionPolicy;
 	TUniqueFunction<void()> ExecFunction;
 	TTypeCompatibleBytes<ResultType> Result;
-	std::atomic<int> NumResultsSet = 0;
+	std::atomic<int32> NumResultsSet = 0;
 	bool bResultConstructed = false;
 };
 
@@ -962,7 +941,7 @@ public:
 			TSharedPtr<TOnlineAsyncOp<OpType>> PinnedOperation = WeakOperation.Pin();
 			if (PinnedOperation)
 			{
-				int ResultIndex = 0;
+				int32 ResultIndex = 0;
 				new(this->Result.GetTypedPtr()) ResultType();
 				this->GetResultRef().SetNum(LastResult.Num());
 				this->bResultConstructed = true;
@@ -974,18 +953,17 @@ public:
 						TPromise<ResultElementType> Promise;
 						// set promise continuation before calling the callable so that we will complete the step as soon as the value is set
 						Promise.GetFuture()
-							.Next([this, WeakOperation, ThisResultIndex = ResultIndex++](const ResultElementType& Value)
+						.Next([this, WeakOperation, ThisResultIndex = ResultIndex++](const ResultElementType& Value)
+						{
+							if (TSharedPtr<TOnlineAsyncOp<OpType>> PinnedOperation2 = WeakOperation.Pin())
+							{
+								this->GetResultRef()[ThisResultIndex] = Value;
+								if (++this->NumResultsSet == this->GetResultRef().Num())
 								{
-									TSharedPtr<TOnlineAsyncOp<OpType>> PinnedOperation2 = WeakOperation.Pin();
-									if (PinnedOperation2)
-									{
-										this->GetResultRef()[ThisResultIndex] = Value;
-										if (++this->NumResultsSet == this->GetResultRef().Num())
-										{
-											PinnedOperation2->ExecuteNextStep();
-										}
-									}
-								});
+									PinnedOperation2->ExecuteNextStep();
+								}
+							}
+						});
 
 						Callable(*PinnedOperation, MoveTempIfPossible(LastResultElement), MoveTemp(Promise));
 					}
@@ -995,18 +973,17 @@ public:
 					for (auto&& LastResultElement : LastResult)
 					{
 						Callable(*PinnedOperation, MoveTempIfPossible(LastResultElement))
-							.Next([this, WeakOperation, ThisResultIndex = ResultIndex++](const ResultElementType& Value)
+						.Next([this, WeakOperation, ThisResultIndex = ResultIndex++](const ResultElementType& Value)
+						{
+							if (TSharedPtr<TOnlineAsyncOp<OpType>> PinnedOperation2 = WeakOperation.Pin())
+							{
+								this->GetResultRef()[ThisResultIndex] = Value;
+								if (++this->NumResultsSet == this->GetResultRef().Num())
 								{
-									TSharedPtr<TOnlineAsyncOp<OpType>> PinnedOperation2 = WeakOperation.Pin();
-									if (PinnedOperation2)
-									{
-										this->GetResultRef()[ThisResultIndex] = Value;
-										if (++this->NumResultsSet == this->GetResultRef().Num())
-										{
-											PinnedOperation2->ExecuteNextStep();
-										}
-									}
-								});
+									PinnedOperation2->ExecuteNextStep();
+								}
+							}
+						});
 					}
 				}
 				else
@@ -1022,13 +999,40 @@ public:
 	}
 };
 
-template <typename ResultElementType>
+namespace Private
+{
+// Primary template with a dependent static_assert that only fails when instantiated
+template <typename ResultType, typename = void>
+struct TForEachNResultElementType
+{
+	// Make the static_assert dependent on the template parameter to delay evaluation
+	static_assert(sizeof(ResultType) == 0, "Unsupported result type for TForEachNResultElementType");
+};
+
+template <typename ResultType>
+struct TForEachNResultElementType<ResultType, std::enable_if_t<TIsContiguousContainer<ResultType>::Value, void >>
+{
+	using ResultElementType = TElementType_T<ResultType>;
+};
+
+template <typename ResultType>
+struct TForEachNResultElementType<ResultType, std::enable_if_t<!TIsContiguousContainer<ResultType>::Value, void>>
+{
+	using ResultElementType = ResultType;
+};
+
+template <typename ResultType>
+using TForEachNResultElementType_T = typename TForEachNResultElementType<ResultType>::ResultElementType;
+
+} // Private
+
+template <typename CallableResultType, typename ResultElementType = Private::TForEachNResultElementType_T<CallableResultType>>
 class TForEachNStep : public TForEachStepBase<ResultElementType>
 {
 public:
 	using typename TForEachStepBase<ResultElementType>::ResultType;
 
-	TForEachNStep(int InBatchSize, FOnlineAsyncExecutionPolicy&& InExecutionPolicy)
+	TForEachNStep(int32 InBatchSize, FOnlineAsyncExecutionPolicy&& InExecutionPolicy)
 		: TForEachStepBase<ResultElementType>(MoveTemp(InExecutionPolicy))
 		, BatchSize(InBatchSize)
 	{
@@ -1039,98 +1043,134 @@ public:
 	{
 		this->ExecFunction = [this, WeakOperation = TWeakPtr<TOnlineAsyncOp<OpType>>(InOperation.AsShared()), &LastResult = InLastResult, Callable = MoveTemp(InCallable)]() mutable
 		{
-			TSharedPtr<TOnlineAsyncOp<OpType>> PinnedOperation = WeakOperation.Pin();
-			if (PinnedOperation)
+			if (TSharedPtr<TOnlineAsyncOp<OpType>> PinnedOperation = WeakOperation.Pin())
 			{
-				int ResultIndex = 0;
 				new(this->Result.GetTypedPtr()) ResultType();
-				this->GetResultRef().SetNumUninitialized(LastResult.Num());
+				const int32 ResultArraySize = TIsContiguousContainer<CallableResultType>::Value ? LastResult.Num()
+					: (LastResult.Num() + BatchSize - 1) / BatchSize;
+				const int32 OutputBatchSize = TIsContiguousContainer<CallableResultType>::Value ? BatchSize : 1;
+				this->GetResultRef().SetNumUninitialized(ResultArraySize);
 				this->bResultConstructed = true;
 
 				if constexpr (TOnlineAsyncOpCallableTraits<CallableType>::bRequiresPromise)
 				{
-					for (int Index = 0; Index < LastResult.Num(); Index += BatchSize)
+					for (int32 Index = 0, OutputIndex = 0; Index < LastResult.Num(); Index += BatchSize, OutputIndex += OutputBatchSize)
 					{
-						int NumElementsInBatch = BatchSize;
+						int32 NumElementsInBatch = BatchSize;
+						int32 NumOutputElementsInBatch = BatchSize;
 						if (Index + BatchSize > LastResult.Num())
 						{
 							NumElementsInBatch = LastResult.Num() - Index;
+							if constexpr (TIsContiguousContainer<CallableResultType>::Value)
+							{
+								NumOutputElementsInBatch = NumElementsInBatch;
+							}
 						}
 
-						TPromise<TArray<ResultElementType>> Promise;
+						TPromise<CallableResultType> Promise;
 						// set promise continuation before calling the callable so that we will complete the step as soon as the value is set
 						Promise.GetFuture()
-							.Next([this, WeakOperation, Index, NumElementsInBatch](const TArray<ResultElementType>& BatchResult)
+						.Next([this, WeakOperation, Index, OutputIndex, NumElementsInBatch, NumOutputElementsInBatch](CallableResultType&& BatchResult)
+						{
+							if (TSharedPtr<TOnlineAsyncOp<OpType>> PinnedOperation2 = WeakOperation.Pin())
+							{
+								if constexpr (TIsContiguousContainer<CallableResultType>::Value)
 								{
-									TSharedPtr<TOnlineAsyncOp<OpType>> PinnedOperation2 = WeakOperation.Pin();
-									if (PinnedOperation2)
+									if (NumElementsInBatch != GetNum(BatchResult))
 									{
-										if (NumElementsInBatch != BatchResult.Num())
-										{
-											PinnedOperation2->SetError(Online::Errors::InvalidResults());
-											return;
-										}
-
-										MoveConstructItems(&this->GetResultRef()[Index], BatchResult.GetData(), BatchResult.Num());
-
-										if ((this->NumResultsSet += NumElementsInBatch) == this->GetResultRef().Num())
-										{
-											PinnedOperation2->ExecuteNextStep();
-										}
+										PinnedOperation2->SetError(Online::Errors::InvalidResults());
+										return;
 									}
-								});
+
+									MoveConstructItems(&this->GetResultRef()[Index], GetData(BatchResult), GetNum(BatchResult));
+
+									if ((this->NumResultsSet += NumElementsInBatch) == this->GetResultRef().Num())
+									{
+										PinnedOperation2->ExecuteNextStep();
+									}
+								}
+								else
+								{
+									::new (&this->GetResultRef()[OutputIndex]) CallableResultType(MoveTemp(BatchResult));
+								}
+
+								if ((this->NumResultsSet += NumOutputElementsInBatch) == this->GetResultRef().Num())
+								{
+									PinnedOperation2->ExecuteNextStep();
+								}
+							}
+						});
 
 						Callable(*PinnedOperation, MakeArrayView(&LastResult[Index], NumElementsInBatch), MoveTemp(Promise));
 					}
 				}
 				else if constexpr (TOnlineAsyncOpCallableTraits<CallableType>::bAsyncResult)
 				{
-					for (int Index = 0; Index < LastResult.Num(); Index += BatchSize)
+					for (int32 Index = 0, OutputIndex = 0; Index < LastResult.Num(); Index += BatchSize, OutputIndex += OutputBatchSize)
 					{
-						int NumElementsInBatch = BatchSize;
+						int32 NumElementsInBatch = BatchSize;
+						int32 NumOutputElementsInBatch = OutputBatchSize;
 						if (Index + BatchSize > LastResult.Num())
 						{
 							NumElementsInBatch = LastResult.Num() - Index;
+							if constexpr (TIsContiguousContainer<CallableResultType>::Value)
+							{
+								NumOutputElementsInBatch = NumElementsInBatch;
+							}
 						}
 
 						Callable(*PinnedOperation, MakeArrayView(&LastResult[Index], NumElementsInBatch))
-							.Next([this, WeakOperation, Index, NumElementsInBatch](const TArray<ResultElementType>& BatchResult)
+						.Next([this, WeakOperation, Index, OutputIndex, NumElementsInBatch, NumOutputElementsInBatch](CallableResultType&& BatchResult)
+						{
+							if (TSharedPtr<TOnlineAsyncOp<OpType>> PinnedOperation2 = WeakOperation.Pin())
+							{
+								if constexpr (TIsContiguousContainer<CallableResultType>::Value)
 								{
-									TSharedPtr<TOnlineAsyncOp<OpType>> PinnedOperation2 = WeakOperation.Pin();
-									if (PinnedOperation2)
+									if (NumElementsInBatch != GetNum(BatchResult))
 									{
-										if (NumElementsInBatch != BatchResult.Num())
-										{
-											PinnedOperation2->SetError(Online::Errors::InvalidResults());
-											return;
-										}
-
-										MoveConstructItems(&this->GetResultRef()[Index], BatchResult.GetData(), BatchResult.Num());
-
-										if ((this->NumResultsSet += NumElementsInBatch) == this->GetResultRef().Num())
-										{
-											PinnedOperation2->ExecuteNextStep();
-										}
+										PinnedOperation2->SetError(Online::Errors::InvalidResults());
+										return;
 									}
-								});
+
+									MoveConstructItems(&this->GetResultRef()[OutputIndex], GetData(BatchResult), GetNum(BatchResult));
+								}
+								else
+								{
+									::new (&this->GetResultRef()[OutputIndex]) CallableResultType(MoveTemp(BatchResult));
+								}
+
+								if ((this->NumResultsSet += NumOutputElementsInBatch) == this->GetResultRef().Num())
+								{
+									PinnedOperation2->ExecuteNextStep();
+								}
+							}
+						});
 					}
 				}
 				else
 				{
-					for (int Index = 0; Index < LastResult.Num(); Index += BatchSize)
+					for (int32 Index = 0, OutputIndex = 0; Index < LastResult.Num(); Index += BatchSize, OutputIndex += OutputBatchSize)
 					{
-						int NumElementsInBatch = BatchSize;
+						int32 NumElementsInBatch = BatchSize;
 						if (Index + BatchSize > LastResult.Num())
 						{
 							NumElementsInBatch = LastResult.Num() - Index;
 						}
-						TArray<ResultElementType> BatchResult = Callable(*PinnedOperation, MakeArrayView(&LastResult[Index], NumElementsInBatch));
-						if (NumElementsInBatch != BatchResult.Num())
+
+						if constexpr (TIsContiguousContainer<CallableResultType>::Value)
+						{
+							CallableResultType BatchResult = Callable(*PinnedOperation, MakeArrayView(&LastResult[Index], NumElementsInBatch));
+							if (NumElementsInBatch != GetNum(BatchResult))
+							{
+								PinnedOperation->SetError(Online::Errors::InvalidResults());
+								break;
+							}
+							MoveConstructItems(&this->GetResultRef()[OutputIndex], GetData(BatchResult), GetNum(BatchResult));
+						}
+						else
 						{
-							PinnedOperation->SetError(Online::Errors::InvalidResults());
-							break;
+							::new (&this->GetResultRef()[OutputIndex]) CallableResultType(Callable(*PinnedOperation, MakeArrayView(&LastResult[Index], NumElementsInBatch)));
 						}
-						MoveConstructItems(&this->GetResultRef()[Index], BatchResult.GetData(), BatchResult.Num());
 					}
 					PinnedOperation->ExecuteNextStep();
 				}
@@ -1139,7 +1179,7 @@ public:
 	}
 
 private:
-	int BatchSize;
+	int32 BatchSize;
 };
 
 // Provides Then continuation for both TOnlineAsyncOp and TOnlineChainableAsyncOp
@@ -1199,7 +1239,7 @@ public:
 	//     LastResultElements is a TArray<LastResultElementType> or const TArray<LastResultElementType>&
 	//     ResultType is any non-void type
 	template <typename CallableType> // previous value was TArray
-	auto ForEachN(int BatchSize, CallableType&& Callable, FOnlineAsyncExecutionPolicy ExecutionPolicy = FOnlineAsyncExecutionPolicy::RunOnGameThread());
+	auto ForEachN(int32 BatchSize, CallableType&& Callable, FOnlineAsyncExecutionPolicy ExecutionPolicy = FOnlineAsyncExecutionPolicy::RunOnGameThread());
 
 protected:
 	LastResultType& LastResult;
@@ -1402,7 +1442,7 @@ public:
 		}
 	}
 
-	virtual void SetError(FOnlineError&& Error)
+	virtual void SetError(FOnlineError&& Error) override
 	{
 		if (SharedState->State < EAsyncOpState::Complete)
 		{
@@ -1410,7 +1450,7 @@ public:
 		}
 	}
 
-	FOnlineServicesCommon& GetServices() { return Services; }
+	FOnlineServicesCommon& GetServices() const { return Services; }
 
 	template <typename QueueType>
 	void Enqueue(QueueType& Queue)
@@ -1433,15 +1473,14 @@ public:
 	{
 		if (!IsComplete())
 		{
-			const int StepToExecute = NextStep;
+			const int32 StepToExecute = NextStep;
 			++NextStep;
 			if (StepToExecute < Steps.Num())
 			{
 				Execute(Steps[StepToExecute]->GetExecutionPolicy(),
 					[this, StepToExecute, WeakThis = TWeakPtr<TOnlineAsyncOp<OpType>>(this->AsShared())]()
 					{
-						TSharedPtr<TOnlineAsyncOp<OpType>> PinnedThis = WeakThis.Pin();
-						if (PinnedThis)
+						if (TSharedPtr<TOnlineAsyncOp<OpType>> PinnedThis = WeakThis.Pin())
 						{
 							Steps[StepToExecute]->Execute();
 						}
@@ -1454,14 +1493,13 @@ public:
 	{
 		if (!IsComplete())
 		{
-			const int StepToExecute = NextStep == 0 ? 0 : NextStep - 1;
+			const int32 StepToExecute = NextStep == 0 ? 0 : NextStep - 1;
 			if (StepToExecute < Steps.Num())
 			{
 				Execute(Steps[StepToExecute]->GetExecutionPolicy(),
 					[this, StepToExecute, WeakThis = TWeakPtr<TOnlineAsyncOp<OpType>>(this->AsShared())]()
 					{
-						TSharedPtr<TOnlineAsyncOp<OpType>> PinnedThis = WeakThis.Pin();
-						if (PinnedThis)
+						if (TSharedPtr<TOnlineAsyncOp<OpType>> PinnedThis = WeakThis.Pin())
 						{
 							Steps[StepToExecute]->Execute();
 						}
@@ -1539,8 +1577,7 @@ protected:
 			Execute(FOnlineAsyncExecutionPolicy::RunOnGameThread(),
 				[State, LambdaResult = MoveTemp(Result), WeakThis = TWeakPtr<TOnlineAsyncOp<OpType>>(this->AsShared())]() mutable
 				{
-					TSharedPtr<TOnlineAsyncOp<OpType>> PinnedThis = WeakThis.Pin();
-					if (PinnedThis)
+					if (TSharedPtr<TOnlineAsyncOp<OpType>> PinnedThis = WeakThis.Pin())
 					{
 						PinnedThis->SetResultAndState(MoveTemp(LambdaResult), State);
 					}
@@ -1689,7 +1726,7 @@ protected:
 	TArray<TUniquePtr<Private::IStep>> Steps;
 	TOnlineEventCallable<void(const TOnlineAsyncOp<OpType>&)> OnStartEvent;
 	TOnlineEventCallable<void(const TOnlineAsyncOp<OpType>&, const TOnlineResult<OpType>&)> OnCompleteEvent;
-	int NextStep = 0;
+	int32 NextStep = 0;
 	double OpStartTimeAbsoluteSeconds;
 
 	friend class FOnlineAsyncOpCache;
@@ -1746,19 +1783,19 @@ auto TOnlineAsyncOpBase<Outer, OpType, LastResultType>::ForEach(CallableType&& I
 
 template <typename Outer, typename OpType, typename LastResultType>
 template <typename CallableType>
-auto TOnlineAsyncOpBase<Outer, OpType, LastResultType>::ForEachN(int BatchSize, CallableType&& InCallable, FOnlineAsyncExecutionPolicy ExecutionPolicy)
+auto TOnlineAsyncOpBase<Outer, OpType, LastResultType>::ForEachN(int32 BatchSize, CallableType&& InCallable, FOnlineAsyncExecutionPolicy ExecutionPolicy)
 {
-	using ResultTypeElement = TElementType_T<typename TOnlineAsyncOpCallableTraits<CallableType>::ResultType>;
+	using CallableResultType = typename TOnlineAsyncOpCallableTraits<CallableType>::ResultType;
 
 	TOnlineAsyncOp<OpType>& Op = static_cast<Outer*>(this)->GetOwningOperation();
 
-	TForEachNStep<ResultTypeElement>* Step = new TForEachNStep<ResultTypeElement>(BatchSize, MoveTemp(ExecutionPolicy));
+	TForEachNStep<CallableResultType>* Step = new TForEachNStep<CallableResultType>(BatchSize, MoveTemp(ExecutionPolicy));
 	TUniquePtr<IStep> StepPtr(Step);
 	Step->SetExecFunction(Op, LastResult, MoveTemp(InCallable));
 
 	Op.AddStep(MoveTemp(StepPtr));
 
-	return TOnlineChainableAsyncOp<OpType, TArray<ResultTypeElement>>(Op, Step->GetResultRef());
+	return TOnlineChainableAsyncOp<OpType, typename TForEachNStep<CallableResultType>::ResultType>(Op, Step->GetResultRef());
 }
 
 
