Add interface name into params of OnOnlineAsyncOpCompleted delegate.

diff --git a/Engine/Plugins/Online/OnlineServices/Source/OnlineServicesCommon/Private/Online/OnlineServicesCommonDelegates.cpp b/Engine/Plugins/Online/OnlineServices/Source/OnlineServicesCommon/Private/Online/OnlineServicesCommonDelegates.cpp
new file mode 100644
--- /dev/null
+++ b/Engine/Plugins/Online/OnlineServices/Source/OnlineServicesCommon/Private/Online/OnlineServicesCommonDelegates.cpp
@@ -0,0 +1,16 @@
+// Copyright Epic Games, Inc. All Rights Reserved.
+
+#include "Online/OnlineServicesCommonDelegates.h"
+#include "Online/OnlineServicesCommon.h"
+
+namespace UE::Online {
+
+FOnOnlineAsyncOpCompletedParams::FOnOnlineAsyncOpCompletedParams(FOnlineServicesCommon& InOnlineServicesCommon, const TOptional<FOnlineError>& InOnlineError)
+	: OnlineServicesCommon(InOnlineServicesCommon.AsShared())
+	, OnlineError(InOnlineError)
+{
+}
+
+FOnOnlineAsyncOpCompletedV2 OnOnlineAsyncOpCompletedV2;
+
+/* UE::Online */ }
diff --git a/Engine/Plugins/Online/OnlineServices/Source/OnlineServicesCommon/Public/Online/OnlineAsyncOp.h b/Engine/Plugins/Online/OnlineServices/Source/OnlineServicesCommon/Public/Online/OnlineAsyncOp.h
--- a/Engine/Plugins/Online/OnlineServices/Source/OnlineServicesCommon/Public/Online/OnlineAsyncOp.h
+++ b/Engine/Plugins/Online/OnlineServices/Source/OnlineServicesCommon/Public/Online/OnlineAsyncOp.h
@@ -6,6 +6,7 @@
 #include "Containers/Ticker.h"
 #include "Online/OnlineAsyncOpHandle.h"
 #include "Online/OnlineResult.h"
+#include "Online/OnlineServicesCommonDelegates.h"
 #include "Online/OnlineServicesDelegates.h"
 #include "Online/OnlineServicesLog.h"
 #include "Online/OnlineTypeInfo.h"
@@ -1554,6 +1555,8 @@ public:
 	TOnlineEvent<void(const TOnlineAsyncOp<OpType>&)> OnStart() { return OnStartEvent; }
 	TOnlineEvent<void(const TOnlineAsyncOp<OpType>&, const TOnlineResult<OpType>&)> OnComplete() { return OnCompleteEvent; }
 
+	void SetInterfaceName(FStringView InName) { InterfaceName = InName; }
+
 protected:
 	FOnlineServicesCommon& Services;
 
@@ -1593,11 +1596,25 @@ protected:
 			SharedHandleState->TriggerOnComplete(Result);
 		}
 
-		// The general callback OnOnlineAsyncOpCompleted needs to be called before OnCompleteEvent, because 
+		// The general callback OnOnlineAsyncOpCompletedV2 needs to be called before OnCompleteEvent, because 
 		// user code triggered by OnCompleteEvent could potentially remove the last reference of OnlineServices
-		// instance and destroy it, which will make it invalid for OnOnlineAsyncOpCompleted callback
+		// instance and destroy it, which will make it invalid for OnOnlineAsyncOpCompletedV2 callback
 		double DurationInSeconds = FPlatformTime::Seconds() - OpStartTimeAbsoluteSeconds;
-		OnOnlineAsyncOpCompleted.Broadcast(OpType::Name, Services, Result.IsOk() ? FOnlineError(Errors::ErrorCode::Common::Success) : Result.GetErrorValue(), DurationInSeconds);
+		TOptional<FOnlineError> ErrorValue;
+		if (Result.IsError())
+		{
+			ErrorValue = Result.GetErrorValue();
+		}
+
+PRAGMA_DISABLE_DEPRECATION_WARNINGS
+		OnOnlineAsyncOpCompleted.Broadcast(OpType::Name, Services, ErrorValue.Get(FOnlineError(Errors::ErrorCode::Common::Success)), DurationInSeconds);
+PRAGMA_ENABLE_DEPRECATION_WARNINGS
+
+		FOnOnlineAsyncOpCompletedParams Params(Services, ErrorValue);
+		Params.DurationInSeconds = DurationInSeconds;
+		Params.InterfaceName = InterfaceName;
+		Params.OpName = OpType::Name;
+		OnOnlineAsyncOpCompletedV2.Broadcast(Params);
 
 		OnCompleteEvent.Broadcast(*this, Result);
 	}
@@ -1731,6 +1748,7 @@ protected:
 
 	friend class FOnlineAsyncOpCache;
 	FOnlineEventDelegateHandle OpCacheHandle; // Delegate handle for FOnlineAsyncOpCache internal usage
+	FString InterfaceName;
 };
 
 template <typename OpType>
diff --git a/Engine/Plugins/Online/OnlineServices/Source/OnlineServicesCommon/Public/Online/OnlineComponent.h b/Engine/Plugins/Online/OnlineServices/Source/OnlineServicesCommon/Public/Online/OnlineComponent.h
--- a/Engine/Plugins/Online/OnlineServices/Source/OnlineServicesCommon/Public/Online/OnlineComponent.h
+++ b/Engine/Plugins/Online/OnlineServices/Source/OnlineServicesCommon/Public/Online/OnlineComponent.h
@@ -52,19 +52,25 @@ public:
 	template <typename OpType>
 	TOnlineAsyncOpRef<OpType> GetOp(typename OpType::Params&& Params)
 	{
-		return Services.template GetOp<OpType>(MoveTemp(Params), GetConfigSectionHeirarchy(OpType::Name));
+		TOnlineAsyncOpRef<OpType> Op = Services.template GetOp<OpType>(MoveTemp(Params), GetConfigSectionHeirarchy(OpType::Name));
+		Op->SetInterfaceName(InterfaceName);
+		return Op;
 	}
 
 	template <typename OpType, typename ParamsFuncsType = TJoinableOpParamsFuncs<OpType>>
 	TOnlineAsyncOpRef<OpType> GetJoinableOp(typename OpType::Params&& Params)
 	{
-		return Services.template GetJoinableOp<OpType, ParamsFuncsType>(MoveTemp(Params), GetConfigSectionHeirarchy(OpType::Name));
+		TOnlineAsyncOpRef<OpType> Op = Services.template GetJoinableOp<OpType, ParamsFuncsType>(MoveTemp(Params), GetConfigSectionHeirarchy(OpType::Name));
+		Op->SetInterfaceName(InterfaceName);
+		return Op;
 	}
 
 	template <typename OpType, typename ParamsFuncsType = TMergeableOpParamsFuncs<OpType>>
 	TOnlineAsyncOpRef<OpType> GetMergeableOp(typename OpType::Params&& Params)
 	{
-		return Services.template GetMergeableOp<OpType, ParamsFuncsType>(MoveTemp(Params), GetConfigSectionHeirarchy(OpType::Name));
+		TOnlineAsyncOpRef<OpType> Op = Services.template GetMergeableOp<OpType, ParamsFuncsType>(MoveTemp(Params), GetConfigSectionHeirarchy(OpType::Name));
+		Op->SetInterfaceName(InterfaceName);
+		return Op;
 	}
 
 	template <typename ServicesType = FOnlineServicesCommon>
diff --git a/Engine/Plugins/Online/OnlineServices/Source/OnlineServicesCommon/Public/Online/OnlineServicesCommonDelegates.h b/Engine/Plugins/Online/OnlineServices/Source/OnlineServicesCommon/Public/Online/OnlineServicesCommonDelegates.h
new file mode 100644
--- /dev/null
+++ b/Engine/Plugins/Online/OnlineServices/Source/OnlineServicesCommon/Public/Online/OnlineServicesCommonDelegates.h
@@ -0,0 +1,45 @@
+// Copyright Epic Games, Inc. All Rights Reserved.
+
+#pragma once
+
+#include "Delegates/Delegate.h"
+#include "Online/OnlineError.h"
+
+namespace UE::Online {
+
+class FOnlineServicesCommon;
+
+// Notification params of OnOnlineAsyncOpCompletedV2 delegate
+struct ONLINESERVICESCOMMON_API FOnOnlineAsyncOpCompletedParams
+{
+	/**
+	 * FOnOnlineAsyncOpCompletedParams ctor
+	 *
+	 * @param OnlineServicesCommon the OnlineServices instance
+	 * @param OnlineError the result of completed operation
+	 */
+	FOnOnlineAsyncOpCompletedParams(FOnlineServicesCommon& InOnlineServicesCommon, const TOptional<FOnlineError>& InOnlineError);
+
+	// The name of completed operation
+	FString OpName;
+	// The name of interface
+	FString InterfaceName;
+	// The OnlineServices instance
+	TWeakPtr<FOnlineServicesCommon> OnlineServicesCommon;
+	// The result of completed operation
+	TOptional<UE::Online::FOnlineError> OnlineError;
+	// The duration of the operation from start to complete
+	double DurationInSeconds = -1.0;
+};
+
+/**
+ * Notification that an online operation has completed
+ *
+ * !!!NOTE!!! The notification can happen on off-game threads, make sure the callbacks are thread-safe
+ *
+ * @param Params the struct FOnOnlineAsyncOpCompletedParams which contains related info
+ */
+DECLARE_MULTICAST_DELEGATE_OneParam(FOnOnlineAsyncOpCompletedV2, const FOnOnlineAsyncOpCompletedParams& Params);
+extern ONLINESERVICESCOMMON_API FOnOnlineAsyncOpCompletedV2 OnOnlineAsyncOpCompletedV2;
+
+/* UE::Online */ }
diff --git a/Engine/Plugins/Online/OnlineServices/Source/OnlineServicesInterface/Private/Online/OnlineServicesDelegates.cpp b/Engine/Plugins/Online/OnlineServices/Source/OnlineServicesInterface/Private/Online/OnlineServicesDelegates.cpp
--- a/Engine/Plugins/Online/OnlineServices/Source/OnlineServicesInterface/Private/Online/OnlineServicesDelegates.cpp
+++ b/Engine/Plugins/Online/OnlineServices/Source/OnlineServicesInterface/Private/Online/OnlineServicesDelegates.cpp
@@ -5,6 +5,8 @@
 namespace UE::Online {
 
 FOnOnlineServicesCreated OnOnlineServicesCreated;
+PRAGMA_DISABLE_DEPRECATION_WARNINGS
 FOnOnlineAsyncOpCompleted OnOnlineAsyncOpCompleted;
+PRAGMA_ENABLE_DEPRECATION_WARNINGS
 
 /* UE::Online */ }
diff --git a/Engine/Plugins/Online/OnlineServices/Source/OnlineServicesInterface/Public/Online/OnlineServicesDelegates.h b/Engine/Plugins/Online/OnlineServices/Source/OnlineServicesInterface/Public/Online/OnlineServicesDelegates.h
--- a/Engine/Plugins/Online/OnlineServices/Source/OnlineServicesInterface/Public/Online/OnlineServicesDelegates.h
+++ b/Engine/Plugins/Online/OnlineServices/Source/OnlineServicesInterface/Public/Online/OnlineServicesDelegates.h
@@ -4,6 +4,7 @@
 
 #include "Delegates/Delegate.h"
 #include "Online/CoreOnline.h"
+#include "Online/OnlineError.h"
 
 namespace UE::Online {
 
@@ -33,6 +34,8 @@ extern ONLINESERVICESINTERFACE_API FOnOnlineServicesCreated OnOnlineServicesCrea
  * @param DurationInSeconds the duration of the operation from start to complete
  */
 DECLARE_MULTICAST_DELEGATE_FourParams(FOnOnlineAsyncOpCompleted, const FString& OpName, const FOnlineServicesCommon& OnlineServicesCommon, const UE::Online::FOnlineError& OnlineError, double DurationInSeconds);
+UE_DEPRECATED(5.6, "OnOnlineAsyncOpCompleted has been deprecated, use OnOnlineAsyncOpCompletedV2 instead")
 extern ONLINESERVICESINTERFACE_API FOnOnlineAsyncOpCompleted OnOnlineAsyncOpCompleted;
 
+
 /* UE::Online */ }
