Add support for parsing UE::Online::ErrorCode::ErrorCodeType from FString

diff --git a/Engine/Plugins/Online/OnlineServices/Source/OnlineServicesInterface/Private/Online/OnlineError.cpp b/Engine/Plugins/Online/OnlineServices/Source/OnlineServicesInterface/Private/Online/OnlineError.cpp
--- a/Engine/Plugins/Online/OnlineServices/Source/OnlineServicesInterface/Private/Online/OnlineError.cpp
+++ b/Engine/Plugins/Online/OnlineServices/Source/OnlineServicesInterface/Private/Online/OnlineError.cpp
@@ -1,8 +1,14 @@
 // Copyright Epic Games, Inc. All Rights Reserved.
 #include "Online/OnlineError.h"
+#include "Algo/AllOf.h"
+#include "Algo/AnyOf.h"
+#include "Containers/StringView.h"
+#include "Misc/Optional.h"
+#include "Misc/Parse.h"
 #include "Online/OnlineErrorCode.h"
 #include "Online/OnlineErrorDefinitions.h"
 #include "Serialization/CompactBinaryWriter.h"
+#include "String/ParseTokens.h"
 
 namespace UE::Online{
 	namespace Errors {
@@ -19,6 +25,33 @@ namespace UE::Online{
 				const uint64 Code = ErrorCodeValue(ErrorCode); 
 				return FString::Printf(TEXT("%llx.%llx.%llx"), System, Category, Code);
 			}
+
+			TOptional<ErrorCodeType> FromString(FStringView ErrorCode)
+			{
+				int32 Count = 0;
+				FStringView Fields[3];
+				UE::String::ParseTokens(ErrorCode, TCHAR('.'), [&Fields, &Count](FStringView Token)
+					{ 
+						if (Count < 3)
+						{
+							if (Algo::AllOf(Token, &FChar::IsHexDigit))
+							{
+								Fields[Count] = Token;
+							}
+						}
+						Count++;
+					});
+				
+				if (Count != 3 || Algo::AnyOf(Fields, &FStringView::IsEmpty))
+				{
+					return NullOpt;
+				}
+
+				const uint64 System = FParse::HexNumber64(Fields[0]);
+				const uint64 Category = FParse::HexNumber64(Fields[1]);
+				const uint64 Code = FParse::HexNumber64(Fields[2]);
+				return Create(System, Category, Code);
+			}
 		} /*namespace ErrorCode */
 	} /* namespace Errors */
 
diff --git a/Engine/Plugins/Online/OnlineServices/Source/OnlineServicesInterface/Public/Online/OnlineErrorCode.h b/Engine/Plugins/Online/OnlineServices/Source/OnlineServicesInterface/Public/Online/OnlineErrorCode.h
--- a/Engine/Plugins/Online/OnlineServices/Source/OnlineServicesInterface/Public/Online/OnlineErrorCode.h
+++ b/Engine/Plugins/Online/OnlineServices/Source/OnlineServicesInterface/Public/Online/OnlineErrorCode.h
@@ -2,8 +2,10 @@
 #pragma once
 
 #include "CoreTypes.h"
+#include "Containers/StringFwd.h"
 
 class FString;
+template <typename OptionalType> struct TOptional;
 
 namespace UE::Online::Errors {
 using ErrorCodeType = uint64;
@@ -20,6 +22,7 @@ struct FOnlineErrorCode
 namespace ErrorCode
 {
 	ONLINESERVICESINTERFACE_API FString ToString(ErrorCodeType ErrorCode);
+	ONLINESERVICESINTERFACE_API TOptional<ErrorCodeType> FromString(FStringView ErrorCode);
 
 	constexpr ErrorCodeType Create(uint32 System, uint32 Category, uint32 Code)
 	{
