OnlineServicesCommon
- OnlineAsyncOp Handle wrapper - can be used when two operations share the same implementation to transform the result
- Fix OnlineAsyncOp ForEachN regression which caused one form to not complete

diff --git a/Engine/Plugins/Online/OnlineServices/Source/OnlineServicesCommon/Private/Online/OnlineSyntaxTests.cpp b/Engine/Plugins/Online/OnlineServices/Source/OnlineServicesCommon/Private/Online/OnlineSyntaxTests.cpp
--- a/Engine/Plugins/Online/OnlineServices/Source/OnlineServicesCommon/Private/Online/OnlineSyntaxTests.cpp
+++ b/Engine/Plugins/Online/OnlineServices/Source/OnlineServicesCommon/Private/Online/OnlineSyntaxTests.cpp
@@ -163,6 +163,13 @@ struct FTestOp
 	struct Result {};
 };
 
+struct FTestOp2
+{
+	static constexpr TCHAR Name[] = TEXT("TestOp2");
+	struct Params {};
+	struct Result {};
+};
+
 inline void TestAsyncOpContinuationSyntax(FOnlineServicesCommon& Services)
 {
 	TSharedRef<TOnlineAsyncOp<FTestOp>> SharedOp = MakeShared<TOnlineAsyncOp<FTestOp>>(Services, FTestOp::Params());
@@ -194,6 +201,10 @@ inline void TestAsyncOpContinuationSyntax(FOnlineServicesCommon& Services)
 		.ForEachN(2, [](TOnlineAsyncOp<FTestOp>&, TArrayView<int> Value) { return TArray<int>(Value); })
 		.ForEachN(2, [](TOnlineAsyncOp<FTestOp>&, TArrayView<int> Value, TPromise<TArray<int>>&& Promise) { return Promise.SetValue(TArray<int>(Value)); })
 		.ForEachN(2, [](TOnlineAsyncOp<FTestOp>&, TArrayView<int> Value) { return MakeFulfilledPromise<TArray<int>>(Value).GetFuture(); })
+		.ForEachN(2, [](TOnlineAsyncOp<FTestOp>&, TArrayView<int> Value) { return 1; })
+		.ForEachN(2, [](TOnlineAsyncOp<FTestOp>&, TArrayView<int> Value) { return MakeFulfilledPromise<int>(1).GetFuture(); })
+		.ForEachN(2, [](TOnlineAsyncOp<FTestOp>&, TArrayView<int> Value, TPromise<int>&& Promise) { Promise.SetValue(1); })
+		.ForEachN(2, [](TOnlineAsyncOp<FTestOp>&, TArrayView<int> Value) { TArray<int> IntArray = { 1, 2 }; return TArray<int>(IntArray); })
 		.Then([](TOnlineAsyncOp<FTestOp>&, const TArray<int>&) { return 1; })
 		.Then([](TOnlineAsyncOp<FTestOp>&, int) {})
 		.Then([](TOnlineAsyncOp<FTestOp>&) { return TContinuationResult<int>::Complete(1); })
@@ -202,6 +213,7 @@ inline void TestAsyncOpContinuationSyntax(FOnlineServicesCommon& Services)
 		.Then([](TOnlineAsyncOp<FTestOp>&) { return TContinuationResult<void>::Complete(); })
 		.Enqueue(Services.GetParallelQueue());
 	Op.GetHandle();
+	Op.GetWrappedHandle<FTestOp2>([](const TOnlineResult<FTestOp>& Result) { return TOnlineResult<FTestOp2>(FTestOp2::Result()); });
 }
 
 /* UE::Online */ }
diff --git a/Engine/Plugins/Online/OnlineServices/Source/OnlineServicesCommon/Public/Online/OnlineAsyncOp.h b/Engine/Plugins/Online/OnlineServices/Source/OnlineServicesCommon/Public/Online/OnlineAsyncOp.h
--- a/Engine/Plugins/Online/OnlineServices/Source/OnlineServicesCommon/Public/Online/OnlineAsyncOp.h
+++ b/Engine/Plugins/Online/OnlineServices/Source/OnlineServicesCommon/Public/Online/OnlineAsyncOp.h
@@ -1084,11 +1084,6 @@ public:
 									}
 
 									MoveConstructItems(&this->GetResultRef()[Index], GetData(BatchResult), GetNum(BatchResult));
-
-									if ((this->NumResultsSet += NumElementsInBatch) == this->GetResultRef().Num())
-									{
-										PinnedOperation2->ExecuteNextStep();
-									}
 								}
 								else
 								{
@@ -1414,6 +1409,12 @@ public:
 		return TOnlineAsyncOpHandle<OpType>(CreateSharedState());
 	}
 
+	template<typename OtherOpType>
+	TOnlineAsyncOpHandle<OtherOpType> GetWrappedHandle(TFunction<TOnlineResult<OtherOpType>(const TOnlineResult<OpType>&)> ResultConversionFunction)
+	{
+		return TOnlineAsyncOpHandle<OtherOpType>(WrapSharedState<OtherOpType>(MoveTemp(ResultConversionFunction)));
+	}
+
 	virtual void ClearCallback() override
 	{
 		UE_LOG(LogOnlineServices, Log, TEXT("%p %s clear complete callbacks while in state %s"), this, OpType::Name, LexToString(SharedState->State));
@@ -1695,6 +1696,20 @@ PRAGMA_ENABLE_DEPRECATION_WARNINGS
 			}
 		}
 
+		void TriggerOnProgress(const FAsyncProgress& Progress)
+		{
+			OnProgressFn.ExecuteIfBound(Progress);
+		}
+
+		void TriggerOnWillRetry(const FWillRetry& RetryDetails)
+		{
+			if (OnWillRetryFn.IsBound())
+			{
+				TSharedRef<TOnlineAsyncOpHandle<OpType>> Handle = TOnlineAsyncOpHandle<OpType>(StaticCastSharedRef<Private::IOnlineAsyncOpSharedState<OpType>>(this->AsShared()));
+				OnWillRetryFn.ExecuteIfBound(Handle, RetryDetails);
+			}
+		}
+
 		void TriggerOnComplete(const TOnlineResult<OpType>& Result)
 		{
 			// TODO: Execute OnCompleteFn next tick on game thread
@@ -1726,6 +1741,35 @@ PRAGMA_ENABLE_DEPRECATION_WARNINGS
 		TWeakPtr<TOnlineAsyncOp<OpType>> AsyncOp;
 	};
 
+	template <typename OtherOpType>	class TWrappedAsyncOpSharedHandleStateContainer;
+
+	template <typename OtherOpType>
+	class TWrappedAsyncOpSharedHandleState : public Private::IOnlineAsyncOpSharedState<OtherOpType>
+	{
+	public:
+		TWrappedAsyncOpSharedHandleState(TWrappedAsyncOpSharedHandleStateContainer<OtherOpType>& InSharedHandleState, TFunction<TOnlineResult<OtherOpType>(const TOnlineResult<OpType>&)>&& InResultConversionFunction);
+		virtual void Cancel(const FOnlineError& Reason) override;
+		virtual EAsyncOpState GetState() const override;
+		virtual void SetOnProgress(TDelegate<void(const FAsyncProgress&)>&& Function) override;
+		virtual void SetOnWillRetry(TDelegate<void(TOnlineAsyncOpHandle<OtherOpType>&, const FWillRetry&)>&& WillRetryFunction) override;
+		virtual void SetOnComplete(TDelegate<void(const TOnlineResult<OtherOpType>&)>&& OnCompleteFunction) override;
+
+	protected:
+		TWrappedAsyncOpSharedHandleStateContainer<OtherOpType>& SharedHandleState;
+		TFunction<TOnlineResult<OtherOpType>(const TOnlineResult<OpType>&)> ResultConversionFunction;
+	};
+
+	template <typename OtherOpType>
+	class TWrappedAsyncOpSharedHandleStateContainer : public FAsyncOpSharedHandleState
+	{
+	public:
+		TWrappedAsyncOpSharedHandleStateContainer(const TSharedRef<TOnlineAsyncOp<OpType>>& InAsyncOp, TFunction<TOnlineResult<OtherOpType>(const TOnlineResult<OpType>&)>&& InResultConversionFunction);
+		TSharedRef<TWrappedAsyncOpSharedHandleState<OtherOpType>> GetWrappedState();
+
+	private:
+		TWrappedAsyncOpSharedHandleState<OtherOpType> WrappedState;
+	};
+
 	void Detach(const TSharedRef<FAsyncOpSharedHandleState>& SharedHandleState)
 	{
 		SharedHandleStates.Remove(SharedHandleState);
@@ -1738,6 +1782,14 @@ PRAGMA_ENABLE_DEPRECATION_WARNINGS
 		return StaticCastSharedRef<Private::IOnlineAsyncOpSharedState<OpType>>(SharedHandleState);
 	}
 
+	template <typename OtherOpType>
+	TSharedRef<Private::IOnlineAsyncOpSharedState<OtherOpType>> WrapSharedState(TFunction<TOnlineResult<OtherOpType>(const TOnlineResult<OpType>&)>&& ResultConversionFunction)
+	{
+		TSharedRef<TWrappedAsyncOpSharedHandleStateContainer<OtherOpType>> SharedHandleState = MakeShared<TWrappedAsyncOpSharedHandleStateContainer<OtherOpType>>(this->AsShared(), MoveTemp(ResultConversionFunction));
+		SharedHandleStates.Add(SharedHandleState);
+		return SharedHandleState->GetWrappedState();
+	}
+
 	TSharedRef<FAsyncOpSharedState> SharedState;
 	TArray<TSharedRef<FAsyncOpSharedHandleState>> SharedHandleStates;
 	TArray<TUniquePtr<Private::IStep>> Steps;
@@ -1751,6 +1803,82 @@ PRAGMA_ENABLE_DEPRECATION_WARNINGS
 	FString InterfaceName;
 };
 
+template <typename OpType>
+template <typename OtherOpType>
+TOnlineAsyncOp<OpType>::TWrappedAsyncOpSharedHandleState<OtherOpType>::TWrappedAsyncOpSharedHandleState(TOnlineAsyncOp<OpType>::TWrappedAsyncOpSharedHandleStateContainer<OtherOpType>& InSharedHandleState, TFunction<TOnlineResult<OtherOpType>(const TOnlineResult<OpType>&)>&& InResultConversionFunction)
+	: SharedHandleState(InSharedHandleState)
+	, ResultConversionFunction(MoveTemp(InResultConversionFunction))
+{
+	check(ResultConversionFunction.IsSet())
+}
+
+template <typename OpType>
+template <typename OtherOpType>
+void TOnlineAsyncOp<OpType>::TWrappedAsyncOpSharedHandleState<OtherOpType>::Cancel(const FOnlineError& Reason)
+{
+	SharedHandleState.Cancel(Reason);
+}
+
+template <typename OpType>
+template <typename OtherOpType>
+EAsyncOpState TOnlineAsyncOp<OpType>::TWrappedAsyncOpSharedHandleState<OtherOpType>::GetState() const
+{
+	return SharedHandleState.GetState();
+}
+
+template <typename OpType>
+template <typename OtherOpType>
+void TOnlineAsyncOp<OpType>::TWrappedAsyncOpSharedHandleState<OtherOpType>::SetOnProgress(TDelegate<void(const FAsyncProgress&)>&& Function)
+{
+	SharedHandleState.SetOnProgress(MoveTemp(Function));
+}
+
+template <typename OpType>
+template <typename OtherOpType>
+void TOnlineAsyncOp<OpType>::TWrappedAsyncOpSharedHandleState<OtherOpType>::SetOnWillRetry(TDelegate<void(TOnlineAsyncOpHandle<OtherOpType>&, const FWillRetry&)>&& WillRetryFunction)
+{
+	SharedHandleState.SetOnWillRetry(TDelegate<void(TOnlineAsyncOpHandle<OpType>&, const FWillRetry&)>::CreateLambda(
+		[this, WillRetryFunction](TOnlineAsyncOpHandle<OpType>& Handle, const FWillRetry& RetryDetails)
+		{
+			if (WillRetryFunction.IsBound())
+			{
+				TOnlineAsyncOpHandle<OtherOpType> WrappedHandle = TOnlineAsyncOpHandle<OtherOpType>(StaticCastSharedRef<Private::IOnlineAsyncOpSharedState<OtherOpType>>(SharedHandleState.GetWrappedState()));
+				WillRetryFunction.Execute(WrappedHandle, RetryDetails);
+			}
+		}));
+}
+
+template <typename OpType>
+template <typename OtherOpType>
+void TOnlineAsyncOp<OpType>::TWrappedAsyncOpSharedHandleState<OtherOpType>::SetOnComplete(TDelegate<void(const TOnlineResult<OtherOpType>&)>&& OnCompleteFunction)
+{
+	SharedHandleState.SetOnComplete(TDelegate<void(const TOnlineResult<OpType>&)>::CreateLambda(
+		[this, OnCompleteFunction](const TOnlineResult<OpType>& Result)
+		{
+			if (OnCompleteFunction.IsBound())
+			{
+				TOnlineResult<OtherOpType> ConvertedResult = ResultConversionFunction(Result);
+				OnCompleteFunction.Execute(ConvertedResult);
+			}
+		}));
+}
+
+template <typename OpType>
+template <typename OtherOpType>
+TOnlineAsyncOp<OpType>::TWrappedAsyncOpSharedHandleStateContainer<OtherOpType>::TWrappedAsyncOpSharedHandleStateContainer(const TSharedRef<TOnlineAsyncOp<OpType>>& InAsyncOp, TFunction<TOnlineResult<OtherOpType>(const TOnlineResult<OpType>&)>&& InResultConversionFunction)
+	: FAsyncOpSharedHandleState(InAsyncOp)
+	, WrappedState(*this, MoveTemp(InResultConversionFunction))
+{
+}
+
+template <typename OpType>
+template <typename OtherOpType>
+TSharedRef<typename TOnlineAsyncOp<OpType>::template TWrappedAsyncOpSharedHandleState<OtherOpType>> TOnlineAsyncOp<OpType>::TWrappedAsyncOpSharedHandleStateContainer<OtherOpType>::GetWrappedState()
+{
+	// Use the aliasing constructor so the lifetime of the shared handle state is tied to the container's lifetime
+	return TSharedPtr<TWrappedAsyncOpSharedHandleState<OtherOpType>>(this->AsShared(), &WrappedState).ToSharedRef();
+}
+
 template <typename OpType>
 using TOnlineAsyncOpRef = TSharedRef<TOnlineAsyncOp<OpType>>;
 template <typename OpType>
diff --git a/Engine/Source/Programs/Online/OnlineTestsCore/Source/Private/Tests/TestAsyncOp.cpp b/Engine/Source/Programs/Online/OnlineTestsCore/Source/Private/Tests/TestAsyncOp.cpp
--- a/Engine/Source/Programs/Online/OnlineTestsCore/Source/Private/Tests/TestAsyncOp.cpp
+++ b/Engine/Source/Programs/Online/OnlineTestsCore/Source/Private/Tests/TestAsyncOp.cpp
@@ -202,11 +202,10 @@ ASYNCOP_TEST_CASE("Verify that the various OnlineAsyncOp continuation forms comp
  		{
  			return MakeFulfilledPromise<TArray<int32>>(Value).GetFuture();
  		})
-		// TODO: Fix this test
- 		/*.ForEachN(3, [](FOnlineAsyncOp&, TArrayView<int32> Value, TPromise<TArray<int32>>&& Promise)
+ 		.ForEachN(3, [](FOnlineAsyncOp&, TArrayView<int32> Value, TPromise<TArray<int32>>&& Promise)
  		{
  			Promise.SetValue(TArray<int32>(Value));
- 		})*/
+ 		})
 		.Then([AsyncResultPtr](TOnlineAsyncOp<FTestOp>& InOp, const TArray<int32>&)
 		{
 			InOp.SetResult(FTestOp::Result());
@@ -246,15 +245,14 @@ ASYNCOP_TEST_CASE("Verify that the various OnlineAsyncOp continuation forms comp
 		{
 			return Async(EAsyncExecution::Thread, [Result = TArray<int32>(Value)]() { return Result; });
 		})
-		// TODO: Fix this test: 
-		/*.ForEachN(3, [](FOnlineAsyncOp&, TArrayView<int32> Value, TPromise<TArray<int32>>&& Promise)
+		.ForEachN(3, [](FOnlineAsyncOp&, TArrayView<int32> Value, TPromise<TArray<int32>>&& Promise)
 		{
 			Async(EAsyncExecution::Thread, [] { return 1; })
 			.Next([MovedPromise = MoveTemp(Promise), Result = TArray<int32>(Value)](int32) mutable
 			{
 				MovedPromise.SetValue(Result);
 			});
-		})*/
+		})
 		.Then([AsyncResultPtr](TOnlineAsyncOp<FTestOp>& Op, const TArray<int32>& Results)
 		{
 			Op.SetResult(FTestOp::Result());
@@ -448,3 +446,87 @@ ASYNCOP_TEST_CASE("Verify that the cache expires once complete)")
 
 	RunToCompletion();
 }
+
+ASYNCOP_TEST_CASE("Verify that wrapped handle completion delegate is called")
+{
+	FTestPipeline& TestPipeline = GetPipeline();
+	TestPipeline.EmplaceAsyncLambda([](FAsyncLambdaResult AsyncResultPtr, const IOnlineServicesPtr& Services)
+	{
+		FOnlineServicesCommon& ServicesCommon = static_cast<FOnlineServicesCommon&>(*Services);
+		TOnlineAsyncOpRef<FTestOpWithParams> Op = ServicesCommon.GetOp<FTestOpWithParams>(FTestOpWithParams::Params({1}));
+		Op->Then([](TOnlineAsyncOp<FTestOpWithParams>& Op) { Op.SetResult(FTestOpWithParams::Result()); })
+		.Enqueue(ServicesCommon.GetParallelQueue());
+
+		Op->GetWrappedHandle<FTestOp>([](const TOnlineResult<FTestOpWithParams>& Result)
+			{
+				if (Result.IsOk())
+				{
+					return TOnlineResult<FTestOp>(FTestOp::Result());
+				}
+				else
+				{
+					return TOnlineResult<FTestOp>(Result.GetErrorValue());
+				}
+			}).OnComplete([AsyncResultPtr](const TOnlineResult<FTestOp>& Result)
+			{
+				CHECK(Result.IsOk());
+				AsyncResultPtr->SetValue(true);
+			});
+	});
+
+	RunToCompletion();
+}
+
+ASYNCOP_TEST_CASE("Verify that wrapped handle completion delegate is called in error case")
+{
+	FTestPipeline& TestPipeline = GetPipeline();
+	TestPipeline.EmplaceAsyncLambda([](FAsyncLambdaResult AsyncResultPtr, const IOnlineServicesPtr& Services)
+	{
+		FOnlineServicesCommon& ServicesCommon = static_cast<FOnlineServicesCommon&>(*Services);
+		TOnlineAsyncOpRef<FTestOpWithParams> Op = ServicesCommon.GetOp<FTestOpWithParams>(FTestOpWithParams::Params({1}));
+		Op->Then([](TOnlineAsyncOp<FTestOpWithParams>& Op) { Op.SetError(UE::Online::Errors::Unknown()); })
+		.Enqueue(ServicesCommon.GetParallelQueue());
+
+		Op->GetWrappedHandle<FTestOp>([](const TOnlineResult<FTestOpWithParams>& Result)
+			{
+				if (Result.IsOk())
+				{
+					return TOnlineResult<FTestOp>(FTestOp::Result());
+				}
+				else
+				{
+					return TOnlineResult<FTestOp>(Result.GetErrorValue());
+				}
+			}).OnComplete([AsyncResultPtr](const TOnlineResult<FTestOp>& Result)
+			{
+				CHECK(Result.IsError());
+				AsyncResultPtr->SetValue(true);
+			});
+	});
+
+	RunToCompletion();
+}
+
+ASYNCOP_TEST_CASE("Verify that wrapped handle completion delegate is called with error in transformed error case")
+{
+	FTestPipeline& TestPipeline = GetPipeline();
+	TestPipeline.EmplaceAsyncLambda([](FAsyncLambdaResult AsyncResultPtr, const IOnlineServicesPtr& Services)
+	{
+		FOnlineServicesCommon& ServicesCommon = static_cast<FOnlineServicesCommon&>(*Services);
+		TOnlineAsyncOpRef<FTestOpWithParams> Op = ServicesCommon.GetOp<FTestOpWithParams>(FTestOpWithParams::Params({1}));
+		Op->Then([](TOnlineAsyncOp<FTestOpWithParams>& Op) { Op.SetResult(FTestOpWithParams::Result()); })
+		.Enqueue(ServicesCommon.GetParallelQueue());
+
+		Op->GetWrappedHandle<FTestOp>([](const TOnlineResult<FTestOpWithParams>& Result)
+			{
+				return TOnlineResult<FTestOp>(UE::Online::Errors::Unknown());
+			}).OnComplete([AsyncResultPtr](const TOnlineResult<FTestOp>& Result)
+			{
+				CHECK(Result.IsError());
+				AsyncResultPtr->SetValue(true);
+			});
+	});
+
+	RunToCompletion();
+}
+
