Log online service creation.

#rnx

diff --git a/Engine/Plugins/Online/OnlineServices/Source/OnlineServicesInterface/Private/Online/OnlineServicesRegistry.cpp b/Engine/Plugins/Online/OnlineServices/Source/OnlineServicesInterface/Private/Online/OnlineServicesRegistry.cpp
--- a/Engine/Plugins/Online/OnlineServices/Source/OnlineServicesInterface/Private/Online/OnlineServicesRegistry.cpp
+++ b/Engine/Plugins/Online/OnlineServices/Source/OnlineServicesInterface/Private/Online/OnlineServicesRegistry.cpp
@@ -188,8 +188,13 @@ TSharedPtr<IOnlineServices> FOnlineServicesRegistry::CreateServices(EOnlineServi
 		Services = FactoryAndPriority->Factory->Create(InstanceName, InstanceConfigName);
 		if (Services)
 		{
+			UE_LOG(LogOnlineServices, Verbose, TEXT("[service_creation_succeeded] OnlineServices=[%s], InstanceName=[%s], InstanceConfigName=[%s]"), LexToString(OnlineServices), *InstanceName.ToString(), *InstanceConfigName.ToString());
 			Services->Init();
 		}
+		else
+		{
+			UE_LOG(LogOnlineServices, Verbose, TEXT("[service_creation_failed] OnlineServices=[%s], InstanceName=[%s], InstanceConfigName=[%s]"), LexToString(OnlineServices), *InstanceName.ToString(), *InstanceConfigName.ToString());
+		}
 	}
 
 	return Services;
