V2  achievements and stats implementations
part 2

diff --git a/Engine/Plugins/Online/OnlineServices/Source/OnlineServicesCommon/Private/Online/StatsCommon.cpp b/Engine/Plugins/Online/OnlineServices/Source/OnlineServicesCommon/Private/Online/StatsCommon.cpp
--- a/Engine/Plugins/Online/OnlineServices/Source/OnlineServicesCommon/Private/Online/StatsCommon.cpp
+++ b/Engine/Plugins/Online/OnlineServices/Source/OnlineServicesCommon/Private/Online/StatsCommon.cpp
@@ -71,6 +71,7 @@ void FStatsCommon::RegisterCommands()
 	RegisterCommand(&FStatsCommon::UpdateStats);
 	RegisterCommand(&FStatsCommon::QueryStats);
 	RegisterCommand(&FStatsCommon::BatchQueryStats);
+	RegisterCommand(&FStatsCommon::TriggerEvent);
 #if !UE_BUILD_SHIPPING
 	RegisterCommand(&FStatsCommon::ResetStats);
 #endif // !UE_BUILD_SHIPPING
@@ -90,6 +91,13 @@ TOnlineAsyncOpHandle<FQueryStats> FStatsCommon::QueryStats(FQueryStats::Params&&
 	return Operation->GetHandle();
 }
 
+TOnlineAsyncOpHandle<FTriggerEvent> FStatsCommon::TriggerEvent(FTriggerEvent::Params&& Params)
+{
+	TOnlineAsyncOpRef<FTriggerEvent> Operation = GetOp<FTriggerEvent>(MoveTemp(Params));
+	Operation->SetError(Errors::NotImplemented());
+	return Operation->GetHandle();
+}
+
 TOnlineAsyncOpHandle<FBatchQueryStats> FStatsCommon::BatchQueryStats(FBatchQueryStats::Params&& Params)
 {
 	TOnlineAsyncOpRef<FBatchQueryStats> Operation = GetOp<FBatchQueryStats>(MoveTemp(Params));
@@ -97,6 +105,7 @@ TOnlineAsyncOpHandle<FBatchQueryStats> FStatsCommon::BatchQueryStats(FBatchQuery
 	return Operation->GetHandle();
 }
 
+
 #if !UE_BUILD_SHIPPING
 TOnlineAsyncOpHandle<FResetStats> FStatsCommon::ResetStats(FResetStats::Params&& Params)
 {
diff --git a/Engine/Plugins/Online/OnlineServices/Source/OnlineServicesCommon/Public/Online/StatsCommon.h b/Engine/Plugins/Online/OnlineServices/Source/OnlineServicesCommon/Public/Online/StatsCommon.h
--- a/Engine/Plugins/Online/OnlineServices/Source/OnlineServicesCommon/Public/Online/StatsCommon.h
+++ b/Engine/Plugins/Online/OnlineServices/Source/OnlineServicesCommon/Public/Online/StatsCommon.h
@@ -86,6 +86,7 @@ public:
 	UE_API virtual TOnlineAsyncOpHandle<FUpdateStats> UpdateStats(FUpdateStats::Params&& Params) override;
 	UE_API virtual TOnlineAsyncOpHandle<FQueryStats> QueryStats(FQueryStats::Params&& Params) override;
 	UE_API virtual TOnlineAsyncOpHandle<FBatchQueryStats> BatchQueryStats(FBatchQueryStats::Params&& Params) override;
+	UE_API virtual TOnlineAsyncOpHandle<FTriggerEvent> TriggerEvent(FTriggerEvent::Params&& Params) override;
 #if !UE_BUILD_SHIPPING
 	UE_API virtual TOnlineAsyncOpHandle<FResetStats> ResetStats(FResetStats::Params&& Params) override;
 #endif // !UE_BUILD_SHIPPING
diff --git a/Engine/Plugins/Online/OnlineServices/Source/OnlineServicesInterface/Public/Online/Stats.h b/Engine/Plugins/Online/OnlineServices/Source/OnlineServicesInterface/Public/Online/Stats.h
--- a/Engine/Plugins/Online/OnlineServices/Source/OnlineServicesInterface/Public/Online/Stats.h
+++ b/Engine/Plugins/Online/OnlineServices/Source/OnlineServicesInterface/Public/Online/Stats.h
@@ -12,6 +12,7 @@ namespace UE::Online { template <typename OpType> class TOnlineResult; }
 namespace UE::Online {
 
 using FStatValue = FSchemaVariant;
+using FOnlineEventParams =  TMap<FString, FSchemaVariant>;
 
 struct FUserStats
 {
@@ -80,6 +81,25 @@ struct FBatchQueryStats
 	};
 };
 
+struct FTriggerEvent
+{
+	static constexpr TCHAR Name[] = TEXT("TriggerEvent");
+
+	struct Params
+	{
+		/* Local user id */
+		FAccountId LocalAccountId;
+		/* The name of the event to fire */
+		FString EventName;
+		/* The events parameters if any */
+		FOnlineEventParams EventParams;
+	};
+
+	struct Result
+	{
+	};
+};
+
 #if !UE_BUILD_SHIPPING
 struct FResetStats
 {
@@ -138,6 +158,13 @@ public:
 	 * @return
 	 */
 	virtual TOnlineAsyncOpHandle<FBatchQueryStats> BatchQueryStats(FBatchQueryStats::Params&& Params) = 0;
+	/**
+	 * Trigger a stat modifying event
+	 *
+	 * @param Parameters for the TriggerEvent call
+	 * @return
+	 */
+	virtual TOnlineAsyncOpHandle<FTriggerEvent> TriggerEvent(FTriggerEvent::Params&& Params) = 0;
 
 #if !UE_BUILD_SHIPPING
 	/**
@@ -187,10 +214,19 @@ BEGIN_ONLINE_STRUCT_META(FQueryStats::Result)
 	ONLINE_STRUCT_FIELD(FQueryStats::Result, Stats)
 END_ONLINE_STRUCT_META()
 
+BEGIN_ONLINE_STRUCT_META(FTriggerEvent::Params)
+	ONLINE_STRUCT_FIELD(FTriggerEvent::Params, LocalAccountId),
+	ONLINE_STRUCT_FIELD(FTriggerEvent::Params, EventName),
+	ONLINE_STRUCT_FIELD(FTriggerEvent::Params, EventParams)
+END_ONLINE_STRUCT_META()
+
+BEGIN_ONLINE_STRUCT_META(FTriggerEvent::Result)
+END_ONLINE_STRUCT_META()
+
 BEGIN_ONLINE_STRUCT_META(FBatchQueryStats::Params)
-	ONLINE_STRUCT_FIELD(FBatchQueryStats::Params, LocalAccountId),
-	ONLINE_STRUCT_FIELD(FBatchQueryStats::Params, TargetAccountIds),
-	ONLINE_STRUCT_FIELD(FBatchQueryStats::Params, StatNames)
+ONLINE_STRUCT_FIELD(FBatchQueryStats::Params, LocalAccountId),
+ONLINE_STRUCT_FIELD(FBatchQueryStats::Params, TargetAccountIds),
+ONLINE_STRUCT_FIELD(FBatchQueryStats::Params, StatNames)
 END_ONLINE_STRUCT_META()
 
 BEGIN_ONLINE_STRUCT_META(FBatchQueryStats::Result)
