- Fix SetExecFunction not working for non copyable types

#rnx

diff --git a/Engine/Plugins/Online/OnlineServices/Source/OnlineServicesCommon/Public/Online/OnlineAsyncOp.h b/Engine/Plugins/Online/OnlineServices/Source/OnlineServicesCommon/Public/Online/OnlineAsyncOp.h
--- a/Engine/Plugins/Online/OnlineServices/Source/OnlineServicesCommon/Public/Online/OnlineAsyncOp.h
+++ b/Engine/Plugins/Online/OnlineServices/Source/OnlineServicesCommon/Public/Online/OnlineAsyncOp.h
@@ -437,11 +437,11 @@ public:
 					TPromise<ResultType> Promise;
 					// set promise continuation before calling the callable so that we will complete the step as soon as the value is set
 					Promise.GetFuture()
-					.Next([this, WeakOperation](const ResultType& Value)
+					.Next([this, WeakOperation](ResultType&& Value)
 					{
 						if (TSharedPtr<TOnlineAsyncOp<OpType>> PinnedOperation2 = WeakOperation.Pin())
 						{
-							EmplaceResult(Value);
+							EmplaceResult(Forward<ResultType>(Value));
 							PinnedOperation2->ExecuteNextStep();
 						}
 					});
@@ -451,11 +451,11 @@ public:
 				else if constexpr (TOnlineAsyncOpCallableTraits<CallableType>::bAsyncResult)
 				{
 					Callable(*PinnedOperation)
-					.Next([this, WeakOperation](const ResultType& Value)
+					.Next([this, WeakOperation](ResultType&& Value)
 					{
 						if (TSharedPtr<TOnlineAsyncOp<OpType>> PinnedOperation2 = WeakOperation.Pin())
 						{
-							EmplaceResult(Value);
+							EmplaceResult(Forward<ResultType>(Value));
 							PinnedOperation2->ExecuteNextStep();
 						}
 					});
