- Add getter for Interface name

#rnx

diff --git a/Engine/Plugins/Online/OnlineServices/Source/OnlineServicesCommon/Public/Online/OnlineAsyncOp.h b/Engine/Plugins/Online/OnlineServices/Source/OnlineServicesCommon/Public/Online/OnlineAsyncOp.h
--- a/Engine/Plugins/Online/OnlineServices/Source/OnlineServicesCommon/Public/Online/OnlineAsyncOp.h
+++ b/Engine/Plugins/Online/OnlineServices/Source/OnlineServicesCommon/Public/Online/OnlineAsyncOp.h
@@ -1556,6 +1556,7 @@ public:
 	TOnlineEvent<void(const TOnlineAsyncOp<OpType>&)> OnStart() { return OnStartEvent; }
 	TOnlineEvent<void(const TOnlineAsyncOp<OpType>&, const TOnlineResult<OpType>&)> OnComplete() { return OnCompleteEvent; }
 
+	const FString& GetInterfaceName() { return InterfaceName; }
 	void SetInterfaceName(FStringView InName) { InterfaceName = InName; }
 
 protected:
