Remove FOnlineServicesCommon dependency from TOnlineAsyncOp, use IOnlineServices instead.

diff --git a/Engine/Plugins/Online/OnlineServices/Source/OnlineServicesCommon/Private/Online/OnlineServicesCommonDelegates.cpp b/Engine/Plugins/Online/OnlineServices/Source/OnlineServicesCommon/Private/Online/OnlineServicesCommonDelegates.cpp
--- a/Engine/Plugins/Online/OnlineServices/Source/OnlineServicesCommon/Private/Online/OnlineServicesCommonDelegates.cpp
+++ b/Engine/Plugins/Online/OnlineServices/Source/OnlineServicesCommon/Private/Online/OnlineServicesCommonDelegates.cpp
@@ -6,11 +6,13 @@
 namespace UE::Online {
 
 FOnOnlineAsyncOpCompletedParams::FOnOnlineAsyncOpCompletedParams(FOnlineServicesCommon& InOnlineServicesCommon, const TOptional<FOnlineError>& InOnlineError)
-	: OnlineServicesCommon(InOnlineServicesCommon.AsShared())
+	: OnlineServicesCommon(StaticCastWeakPtr<FOnlineServicesCommon>(InOnlineServicesCommon.AsWeak()))
 	, OnlineError(InOnlineError)
 {
 }
 
+PRAGMA_DISABLE_DEPRECATION_WARNINGS
 FOnOnlineAsyncOpCompletedV2 OnOnlineAsyncOpCompletedV2;
+PRAGMA_ENABLE_DEPRECATION_WARNINGS
 
 /* UE::Online */ }
diff --git a/Engine/Plugins/Online/OnlineServices/Source/OnlineServicesCommon/Public/Online/OnlineAsyncOp.h b/Engine/Plugins/Online/OnlineServices/Source/OnlineServicesCommon/Public/Online/OnlineAsyncOp.h
--- a/Engine/Plugins/Online/OnlineServices/Source/OnlineServicesCommon/Public/Online/OnlineAsyncOp.h
+++ b/Engine/Plugins/Online/OnlineServices/Source/OnlineServicesCommon/Public/Online/OnlineAsyncOp.h
@@ -9,6 +9,7 @@
 #include "Online/OnlineServicesCommonDelegates.h"
 #include "Online/OnlineServicesDelegates.h"
 #include "Online/OnlineServicesLog.h"
+#include "Online/OnlineServices.h"
 #include "Online/OnlineTypeInfo.h"
 #include "Traits/ElementType.h"
 #include "Traits/IsContiguousContainer.h"
@@ -1366,7 +1367,7 @@ public:
 	using ParamsType = typename OpType::Params;
 	using ResultType = typename OpType::Result;
 
-	TOnlineAsyncOp(FOnlineServicesCommon& InServices, ParamsType&& Params)
+	TOnlineAsyncOp(IOnlineServices& InServices, ParamsType&& Params)
 		: Services(InServices)
 		, SharedState(MakeShared<FAsyncOpSharedState>(MoveTemp(Params)))
 		, OpStartTimeAbsoluteSeconds(FPlatformTime::Seconds())
@@ -1452,7 +1453,13 @@ public:
 		}
 	}
 
-	FOnlineServicesCommon& GetServices() const { return Services; }
+	UE_DEPRECATED(5.6, "GetServices has been deprecated and we will no longer return FOnlineServicesCommon, use GetOnlineServices() instead and convert to FOnlineServicesCommon when needed")
+	FOnlineServicesCommon& GetServices() const
+	{ 
+		return *reinterpret_cast<FOnlineServicesCommon*>(&Services); 
+	}
+
+	IOnlineServices& GetOnlineServices() const { return Services; }
 
 	template <typename QueueType>
 	void Enqueue(QueueType& Queue)
@@ -1560,7 +1567,7 @@ public:
 	void SetInterfaceName(FStringView InName) { InterfaceName = InName; }
 
 protected:
-	FOnlineServicesCommon& Services;
+	IOnlineServices& Services;
 
 	void SetResultAndState(TOnlineResult<OpType>&& Result, EAsyncOpState State)
 	{
@@ -1598,9 +1605,9 @@ protected:
 			SharedHandleState->TriggerOnComplete(Result);
 		}
 
-		// The general callback OnOnlineAsyncOpCompletedV2 needs to be called before OnCompleteEvent, because 
+		// The general callback OnAsyncOpCompleted needs to be called before OnCompleteEvent, because 
 		// user code triggered by OnCompleteEvent could potentially remove the last reference of OnlineServices
-		// instance and destroy it, which will make it invalid for OnOnlineAsyncOpCompletedV2 callback
+		// instance and destroy it, which will make it invalid for OnAsyncOpCompleted callback
 		double DurationInSeconds = FPlatformTime::Seconds() - OpStartTimeAbsoluteSeconds;
 		TOptional<FOnlineError> ErrorValue;
 		if (Result.IsError())
@@ -1609,14 +1616,20 @@ protected:
 		}
 
 PRAGMA_DISABLE_DEPRECATION_WARNINGS
-		OnOnlineAsyncOpCompleted.Broadcast(OpType::Name, Services, ErrorValue.Get(FOnlineError(Errors::ErrorCode::Common::Success)), DurationInSeconds);
-PRAGMA_ENABLE_DEPRECATION_WARNINGS
-
-		FOnOnlineAsyncOpCompletedParams Params(Services, ErrorValue);
+		FOnOnlineAsyncOpCompletedParams Params(GetServices(), ErrorValue);
 		Params.DurationInSeconds = DurationInSeconds;
 		Params.InterfaceName = InterfaceName;
 		Params.OpName = OpType::Name;
 		OnOnlineAsyncOpCompletedV2.Broadcast(Params);
+PRAGMA_ENABLE_DEPRECATION_WARNINGS
+
+		FOnAsyncOpCompletedParams OnAsyncOpCompletedParams;
+		OnAsyncOpCompletedParams.OnlineServices = Services.AsWeak();
+		OnAsyncOpCompletedParams.OnlineError = ErrorValue;
+		OnAsyncOpCompletedParams.DurationInSeconds = DurationInSeconds;
+		OnAsyncOpCompletedParams.InterfaceName = InterfaceName;
+		OnAsyncOpCompletedParams.OpName = OpType::Name;
+		OnAsyncOpCompleted.Broadcast(OnAsyncOpCompletedParams);
 
 		OnCompleteEvent.Broadcast(*this, Result);
 	}
diff --git a/Engine/Plugins/Online/OnlineServices/Source/OnlineServicesCommon/Public/Online/OnlineServicesCommon.h b/Engine/Plugins/Online/OnlineServices/Source/OnlineServicesCommon/Public/Online/OnlineServicesCommon.h
--- a/Engine/Plugins/Online/OnlineServices/Source/OnlineServicesCommon/Public/Online/OnlineServicesCommon.h
+++ b/Engine/Plugins/Online/OnlineServices/Source/OnlineServicesCommon/Public/Online/OnlineServicesCommon.h
@@ -19,7 +19,6 @@ namespace UE::Online {
 
 class FOnlineServicesCommon
 	: public IOnlineServices
-	, public TSharedFromThis<FOnlineServicesCommon>
 	, public FTSTickerObjectBase
 	, public FSelfRegisteringExec
 {
diff --git a/Engine/Plugins/Online/OnlineServices/Source/OnlineServicesCommon/Public/Online/OnlineServicesCommonDelegates.h b/Engine/Plugins/Online/OnlineServices/Source/OnlineServicesCommon/Public/Online/OnlineServicesCommonDelegates.h
--- a/Engine/Plugins/Online/OnlineServices/Source/OnlineServicesCommon/Public/Online/OnlineServicesCommonDelegates.h
+++ b/Engine/Plugins/Online/OnlineServices/Source/OnlineServicesCommon/Public/Online/OnlineServicesCommonDelegates.h
@@ -42,6 +42,7 @@ struct FOnOnlineAsyncOpCompletedParams
  * @param Params the struct FOnOnlineAsyncOpCompletedParams which contains related info
  */
 DECLARE_MULTICAST_DELEGATE_OneParam(FOnOnlineAsyncOpCompletedV2, const FOnOnlineAsyncOpCompletedParams& Params);
+UE_DEPRECATED(5.7, "OnOnlineAsyncOpCompletedV2 has been deprecated, use OnAsyncOpCompleted instead")
 extern ONLINESERVICESCOMMON_API FOnOnlineAsyncOpCompletedV2 OnOnlineAsyncOpCompletedV2;
 
 /* UE::Online */ }
diff --git a/Engine/Plugins/Online/OnlineServices/Source/OnlineServicesInterface/Private/Online/OnlineServicesDelegates.cpp b/Engine/Plugins/Online/OnlineServices/Source/OnlineServicesInterface/Private/Online/OnlineServicesDelegates.cpp
--- a/Engine/Plugins/Online/OnlineServices/Source/OnlineServicesInterface/Private/Online/OnlineServicesDelegates.cpp
+++ b/Engine/Plugins/Online/OnlineServices/Source/OnlineServicesInterface/Private/Online/OnlineServicesDelegates.cpp
@@ -5,8 +5,6 @@
 namespace UE::Online {
 
 FOnOnlineServicesCreated OnOnlineServicesCreated;
-PRAGMA_DISABLE_DEPRECATION_WARNINGS
-FOnOnlineAsyncOpCompleted OnOnlineAsyncOpCompleted;
-PRAGMA_ENABLE_DEPRECATION_WARNINGS
+FOnAsyncOpCompleted OnAsyncOpCompleted;
 
 /* UE::Online */ }
diff --git a/Engine/Plugins/Online/OnlineServices/Source/OnlineServicesInterface/Public/Online/OnlineServices.h b/Engine/Plugins/Online/OnlineServices/Source/OnlineServicesInterface/Public/Online/OnlineServices.h
--- a/Engine/Plugins/Online/OnlineServices/Source/OnlineServicesInterface/Public/Online/OnlineServices.h
+++ b/Engine/Plugins/Online/OnlineServices/Source/OnlineServicesInterface/Public/Online/OnlineServices.h
@@ -49,9 +49,11 @@ struct FGetResolvedConnectString
 	};
 };
 
-class IOnlineServices
+class IOnlineServices : public TSharedFromThis<IOnlineServices>
 {
 public:
+	virtual ~IOnlineServices() {}
+
 	/**
 	 *
 	 */
diff --git a/Engine/Plugins/Online/OnlineServices/Source/OnlineServicesInterface/Public/Online/OnlineServicesDelegates.h b/Engine/Plugins/Online/OnlineServices/Source/OnlineServicesInterface/Public/Online/OnlineServicesDelegates.h
--- a/Engine/Plugins/Online/OnlineServices/Source/OnlineServicesInterface/Public/Online/OnlineServicesDelegates.h
+++ b/Engine/Plugins/Online/OnlineServices/Source/OnlineServicesInterface/Public/Online/OnlineServicesDelegates.h
@@ -8,8 +8,8 @@
 
 namespace UE::Online {
 
-class FOnlineServicesCommon;
 class FOnlineError;
+class IOnlineServices;
 
 /**
  * Online services delegates that are more external to the online services themselves
@@ -23,19 +23,29 @@ class FOnlineError;
 DECLARE_MULTICAST_DELEGATE_OneParam(FOnOnlineServicesCreated, TSharedRef<class IOnlineServices> /*NewServices*/);
 extern ONLINESERVICESINTERFACE_API FOnOnlineServicesCreated OnOnlineServicesCreated;
 
+// Notification params of OnAsyncOpCompleted delegate
+struct FOnAsyncOpCompletedParams
+{
+	// The name of completed operation
+	FString OpName;
+	// The name of interface
+	FString InterfaceName;
+	// The OnlineServices instance
+	TWeakPtr<IOnlineServices> OnlineServices;
+	// The result of completed operation
+	TOptional<UE::Online::FOnlineError> OnlineError;
+	// The duration of the operation from start to complete
+	double DurationInSeconds = -1.0;
+};
+
 /**
  * Notification that an online operation has completed
  *
  * !!!NOTE!!! The notification can happen on off-game threads, make sure the callbacks are thread-safe
  *
- * @param OpName the name of completed operation
- * @param OnlineServicesCommon the OnlineServices instance
- * @param OnlineError the result of completed operation
- * @param DurationInSeconds the duration of the operation from start to complete
+ * @param Params the struct FOnAsyncOpCompletedParams which contains related info
  */
-DECLARE_MULTICAST_DELEGATE_FourParams(FOnOnlineAsyncOpCompleted, const FString& OpName, const FOnlineServicesCommon& OnlineServicesCommon, const UE::Online::FOnlineError& OnlineError, double DurationInSeconds);
-UE_DEPRECATED(5.6, "OnOnlineAsyncOpCompleted has been deprecated, use OnOnlineAsyncOpCompletedV2 instead")
-extern ONLINESERVICESINTERFACE_API FOnOnlineAsyncOpCompleted OnOnlineAsyncOpCompleted;
-
+DECLARE_MULTICAST_DELEGATE_OneParam(FOnAsyncOpCompleted, const FOnAsyncOpCompletedParams& Params);
+extern ONLINESERVICESINTERFACE_API FOnAsyncOpCompleted OnAsyncOpCompleted;
 
 /* UE::Online */ }
