OSSv2 Session interface added online errors (already in session and not in session)
#rnx

diff --git a/Engine/Plugins/Online/OnlineServices/Source/OnlineServicesInterface/Public/Online/OnlineErrorDefinitions.h b/Engine/Plugins/Online/OnlineServices/Source/OnlineServicesInterface/Public/Online/OnlineErrorDefinitions.h
--- a/Engine/Plugins/Online/OnlineServices/Source/OnlineServicesInterface/Public/Online/OnlineErrorDefinitions.h
+++ b/Engine/Plugins/Online/OnlineServices/Source/OnlineServicesInterface/Public/Online/OnlineErrorDefinitions.h
@@ -34,6 +34,8 @@ UE_ONLINE_ERROR_COMMON(Common, Timeout, 20, TEXT("timeout"), LOCTEXT("Timeout",
 UE_ONLINE_ERROR_COMMON(Common, InvalidState, 21, TEXT("invalid_state"), LOCTEXT("InvalidState", "Invalid state"))
 UE_ONLINE_ERROR_COMMON(Common, Unknown, 22, TEXT("unknown"), LOCTEXT("Unknown", "Unknown Error"))
 UE_ONLINE_ERROR_COMMON(Common, NoChange, 23, TEXT("no_change"), LOCTEXT("NoChange", "No change"))
+UE_ONLINE_ERROR_COMMON(Common, AlreadyInSession, 24, TEXT("already_in_session"), LOCTEXT("AlreadyInSession", "Player is already in session"))
+UE_ONLINE_ERROR_COMMON(Common, NotInSession, 25, TEXT("not_in_session"), LOCTEXT("NotInSession", "Player is not in any session"))
 
 } /* namespace UE::Online::Errors */
 
