Add ToAccountId and FromStringData to the ID regristry to support converting account strings to FAccountIDs.
- Some instances don't have an implementation yet, notably OnlineIdPSN and OnlineIdOSSAdapter.h. We can add these later when needed.
#rnx

diff --git a/Engine/Plugins/Online/OnlineServices/Source/OnlineServicesCommon/Public/Online/SessionsCommon.h b/Engine/Plugins/Online/OnlineServices/Source/OnlineServicesCommon/Public/Online/SessionsCommon.h
--- a/Engine/Plugins/Online/OnlineServices/Source/OnlineServicesCommon/Public/Online/SessionsCommon.h
+++ b/Engine/Plugins/Online/OnlineServices/Source/OnlineServicesCommon/Public/Online/SessionsCommon.h
@@ -52,14 +52,19 @@ public:
 	{
 		const FUTF8ToTCHAR IdValueTCHAR((char*)ReplicationData.GetData(), ReplicationData.Num());
 		const FString IdValue = FString::ConstructFromPtrSize(IdValueTCHAR.Get(), IdValueTCHAR.Length());
+		return FromStringData(IdValue);
+	}
 
-		if (!IdValue.IsEmpty())
+	virtual inline FOnlineSessionId FromStringData(const FString& StringData) override
+	{
+		if (!StringData.IsEmpty())
 		{
-			return BasicRegistry.FindOrAddHandle(IdValue);
+			return BasicRegistry.FindOrAddHandle(StringData);
 		}
 
 		return FOnlineSessionId();
 	}
+
 	// End IOnlineSessionIdRegistry
 
 	inline bool IsSessionIdExpired(const FOnlineSessionId& InSessionId) const
@@ -117,14 +122,19 @@ public:
 	{
 		const FUTF8ToTCHAR IdValueTCHAR((char*)ReplicationData.GetData(), ReplicationData.Num());
 		const FString IdValue = FString::ConstructFromPtrSize(IdValueTCHAR.Get(), IdValueTCHAR.Length());
+		return FromStringData(IdValue);
+	}
 
-		if (!IdValue.IsEmpty())
+	virtual inline FSessionInviteId FromStringData(const FString& StringData) override
+	{
+		if (!StringData.IsEmpty())
 		{
-			return BasicRegistry.FindOrAddHandle(IdValue);
+			return BasicRegistry.FindOrAddHandle(StringData);
 		}
 
 		return FSessionInviteId();
 	}
+
 	// End IOnlineSessionIdRegistry
 
 	FOnlineSessionInviteIdStringRegistry(EOnlineServices OnlineServicesType)
diff --git a/Engine/Plugins/Online/OnlineServicesEOS/Source/Private/Online/OnlineIdEOS.cpp b/Engine/Plugins/Online/OnlineServicesEOS/Source/Private/Online/OnlineIdEOS.cpp
--- a/Engine/Plugins/Online/OnlineServicesEOS/Source/Private/Online/OnlineIdEOS.cpp
+++ b/Engine/Plugins/Online/OnlineServicesEOS/Source/Private/Online/OnlineIdEOS.cpp
@@ -127,6 +127,29 @@ FAccountId FOnlineAccountIdRegistryEOS::FromReplicationData(const TArray<uint8>&
 	return FindOrAddAccountId(EpicAccountId, ProductUserId);
 }
 
+FAccountId FOnlineAccountIdRegistryEOS::FromStringData(const FString& StringData)
+{
+	FString EpicAccountIdString;
+	FString ProductUserIdString;
+
+	StringData.Split(TEXT(":"), &EpicAccountIdString, &ProductUserIdString);
+
+	EOS_EpicAccountId EpicAccountId = nullptr;
+	if (!EpicAccountIdString.IsEmpty())
+	{
+		const auto Utf8EpicAccountId = StringCast<UTF8CHAR>(*EpicAccountIdString);
+		EpicAccountId = EOS_EpicAccountId_FromString((const char*)Utf8EpicAccountId.Get());
+	}
+
+	EOS_ProductUserId ProductUserId = nullptr;
+	if (!ProductUserIdString.IsEmpty())
+	{
+		const auto Utf8ProductUserId = StringCast<UTF8CHAR>(*ProductUserIdString);
+		ProductUserId = EOS_ProductUserId_FromString((const char*)Utf8ProductUserId.Get());
+	}
+
+	return FindOrAddAccountId(EpicAccountId, ProductUserId);
+}
 
 FAccountId FOnlineAccountIdRegistryEOS::FindOrAddAccountId(const EOS_ProductUserId ProductUserId)
 {
diff --git a/Engine/Plugins/Online/OnlineServicesEOS/Source/Private/Online/OnlineIdEOS.h b/Engine/Plugins/Online/OnlineServicesEOS/Source/Private/Online/OnlineIdEOS.h
--- a/Engine/Plugins/Online/OnlineServicesEOS/Source/Private/Online/OnlineIdEOS.h
+++ b/Engine/Plugins/Online/OnlineServicesEOS/Source/Private/Online/OnlineIdEOS.h
@@ -33,6 +33,7 @@ public:
 	virtual FString ToLogString(const FAccountId& AccountId) const override;
 	virtual TArray<uint8> ToReplicationData(const FAccountId& AccountId) const override;
 	virtual FAccountId FromReplicationData(const TArray<uint8>& ReplicationString) override;
+	virtual FAccountId FromStringData(const FString& StringData) override;
 	// End IOnlineAccountIdRegistry
 
 	// Begin IOnlineAccountRegistryEOSGS
diff --git a/Engine/Plugins/Online/OnlineServicesEOSGS/Source/Private/Online/OnlineIdEOSGS.cpp b/Engine/Plugins/Online/OnlineServicesEOSGS/Source/Private/Online/OnlineIdEOSGS.cpp
--- a/Engine/Plugins/Online/OnlineServicesEOSGS/Source/Private/Online/OnlineIdEOSGS.cpp
+++ b/Engine/Plugins/Online/OnlineServicesEOSGS/Source/Private/Online/OnlineIdEOSGS.cpp
@@ -84,6 +84,7 @@ TArray<uint8> FOnlineAccountIdRegistryEOSGS::ToReplicationData(const FAccountId&
 	}
 	return ReplicationData;
 }
+
 FAccountId FOnlineAccountIdRegistryEOSGS::FromReplicationData(const TArray<uint8>& ReplicationData)
 {
 	if (ReplicationData.Num() == OnlineIdEOSHexBufferLength)
@@ -96,6 +97,13 @@ FAccountId FOnlineAccountIdRegistryEOSGS::FromReplicationData(const TArray<uint8
 	return Registry.GetInvalidHandle();
 }
 
+FAccountId FOnlineAccountIdRegistryEOSGS::FromStringData(const FString& StringData)
+{
+	const auto Utf8ProductUserId = StringCast<UTF8CHAR>(*StringData);
+	EOS_ProductUserId ProductUserId = EOS_ProductUserId_FromString((const char*)Utf8ProductUserId.Get());
+	return FindOrAddAccountId(ProductUserId);
+}
+
 IOnlineAccountIdRegistryEOSGS& FOnlineAccountIdRegistryEOSGS::GetRegistered()
 {
 	return FOnlineAccountIdRegistryEOSGS::GetRegistered(EOnlineServices::Epic);
diff --git a/Engine/Plugins/Online/OnlineServicesEOSGS/Source/Public/Online/OnlineIdEOSGS.h b/Engine/Plugins/Online/OnlineServicesEOSGS/Source/Public/Online/OnlineIdEOSGS.h
--- a/Engine/Plugins/Online/OnlineServicesEOSGS/Source/Public/Online/OnlineIdEOSGS.h
+++ b/Engine/Plugins/Online/OnlineServicesEOSGS/Source/Public/Online/OnlineIdEOSGS.h
@@ -44,6 +44,7 @@ public:
 	ONLINESERVICESEOSGS_API virtual FString ToLogString(const FAccountId& AccountId) const override;
 	ONLINESERVICESEOSGS_API virtual TArray<uint8> ToReplicationData(const FAccountId& AccountId) const override;
 	ONLINESERVICESEOSGS_API virtual FAccountId FromReplicationData(const TArray<uint8>& ReplicationData) override;
+	ONLINESERVICESEOSGS_API virtual FAccountId FromStringData(const FString& StringData) override;
 	// End IOnlineAccountIdRegistry
 
 	UE_DEPRECATED(5.6, "This method is deprecated, please use the new version taking a EOnlineServices parameter")
diff --git a/Engine/Plugins/Online/OnlineServicesNull/Source/Private/Online/AuthNull.cpp b/Engine/Plugins/Online/OnlineServicesNull/Source/Private/Online/AuthNull.cpp
--- a/Engine/Plugins/Online/OnlineServicesNull/Source/Private/Online/AuthNull.cpp
+++ b/Engine/Plugins/Online/OnlineServicesNull/Source/Private/Online/AuthNull.cpp
@@ -227,10 +227,14 @@ TArray<uint8> FOnlineAccountIdRegistryNull::ToReplicationData(const FAccountId&
 
 FAccountId FOnlineAccountIdRegistryNull::FromReplicationData(const TArray<uint8>& ReplicationData)
 {
-	FString Result = BytesToString(ReplicationData.GetData(), ReplicationData.Num());
-	if(Result.Len() > 0)
+	return FromStringData(BytesToString(ReplicationData.GetData(), ReplicationData.Num()));
+}
+
+FAccountId FOnlineAccountIdRegistryNull::FromStringData(const FString& StringData)
+{
+	if (StringData.Len() > 0)
 	{
-		return FindOrAddAccountId(Result);
+		return FindOrAddAccountId(StringData);
 	}
 	return FAccountId();
 }
diff --git a/Engine/Plugins/Online/OnlineServicesNull/Source/Private/Online/LobbyRegistryNull.cpp b/Engine/Plugins/Online/OnlineServicesNull/Source/Private/Online/LobbyRegistryNull.cpp
--- a/Engine/Plugins/Online/OnlineServicesNull/Source/Private/Online/LobbyRegistryNull.cpp
+++ b/Engine/Plugins/Online/OnlineServicesNull/Source/Private/Online/LobbyRegistryNull.cpp
@@ -75,10 +75,14 @@ TArray<uint8> FOnlineLobbyIdRegistryNull::ToReplicationData(const FLobbyId& Lobb
 
 FLobbyId FOnlineLobbyIdRegistryNull::FromReplicationData(const TArray<uint8>& ReplicationData)
 {
-	FString Result = BytesToString(ReplicationData.GetData(), ReplicationData.Num());
-	if (Result.Len() > 0)
+	return FromStringData(BytesToString(ReplicationData.GetData(), ReplicationData.Num()));
+}
+
+FLobbyId FOnlineLobbyIdRegistryNull::FromStringData(const FString& StringData)
+{
+	if (StringData.Len() > 0)
 	{
-		return FindOrAdd(Result);
+		return FindOrAdd(StringData);
 	}
 	return FLobbyId();
 }
diff --git a/Engine/Plugins/Online/OnlineServicesNull/Source/Public/Online/AuthNull.h b/Engine/Plugins/Online/OnlineServicesNull/Source/Public/Online/AuthNull.h
--- a/Engine/Plugins/Online/OnlineServicesNull/Source/Public/Online/AuthNull.h
+++ b/Engine/Plugins/Online/OnlineServicesNull/Source/Public/Online/AuthNull.h
@@ -27,6 +27,7 @@ public:
 	virtual FString ToLogString(const FAccountId& AccountId) const override;
 	virtual TArray<uint8> ToReplicationData(const FAccountId& AccountId) const override;
 	virtual FAccountId FromReplicationData(const TArray<uint8>& ReplicationString) override;
+	virtual FAccountId FromStringData(const FString& StringData) override;
 	// End IOnlineAccountIdRegistry
 
 	virtual ~FOnlineAccountIdRegistryNull() = default;
diff --git a/Engine/Plugins/Online/OnlineServicesNull/Source/Public/Online/LobbyRegistryNull.h b/Engine/Plugins/Online/OnlineServicesNull/Source/Public/Online/LobbyRegistryNull.h
--- a/Engine/Plugins/Online/OnlineServicesNull/Source/Public/Online/LobbyRegistryNull.h
+++ b/Engine/Plugins/Online/OnlineServicesNull/Source/Public/Online/LobbyRegistryNull.h
@@ -21,6 +21,7 @@ public:
 	virtual FString ToLogString(const FLobbyId& LobbyId) const override;
 	virtual TArray<uint8> ToReplicationData(const FLobbyId& LobbyId) const override;
 	virtual FLobbyId FromReplicationData(const TArray<uint8>& ReplicationString) override;
+	virtual FLobbyId FromStringData(const FString& StringData) override;
 	// End IOnlineAccountIdRegistry
 
 	virtual ~FOnlineLobbyIdRegistryNull() = default;
diff --git a/Engine/Plugins/Online/OnlineServicesOSSAdapter/Source/Public/Online/OnlineIdOSSAdapter.h b/Engine/Plugins/Online/OnlineServicesOSSAdapter/Source/Public/Online/OnlineIdOSSAdapter.h
--- a/Engine/Plugins/Online/OnlineServicesOSSAdapter/Source/Public/Online/OnlineIdOSSAdapter.h
+++ b/Engine/Plugins/Online/OnlineServicesOSSAdapter/Source/Public/Online/OnlineIdOSSAdapter.h
@@ -106,6 +106,12 @@ public:
 	{
 		return TOnlineId<IdType>();
 	}
+
+	virtual TOnlineId<IdType> FromStringData(const FString& StringData) override
+	{
+		return TOnlineId<IdType>();
+	}
+
 	// End IOnlineAccountIdRegistry
 
 private:
diff --git a/Engine/Source/Runtime/CoreOnline/Private/Online/CoreOnline.cpp b/Engine/Source/Runtime/CoreOnline/Private/Online/CoreOnline.cpp
--- a/Engine/Source/Runtime/CoreOnline/Private/Online/CoreOnline.cpp
+++ b/Engine/Source/Runtime/CoreOnline/Private/Online/CoreOnline.cpp
@@ -344,6 +344,17 @@ FAccountId FOnlineIdRegistryRegistry::ToAccountId(EOnlineServices Services, cons
 	return AccountId;
 }
 
+FAccountId FOnlineIdRegistryRegistry::ToAccountId(EOnlineServices Services, const FString& StringData) const
+{
+	FAccountId AccountId;
+	if (IOnlineAccountIdRegistry* Registry = GetAccountIdRegistry(Services))
+	{
+		AccountId = Registry->FromStringData(StringData);
+	}
+
+	return AccountId;
+}
+
 IOnlineAccountIdRegistry* FOnlineIdRegistryRegistry::GetAccountIdRegistry(EOnlineServices OnlineServices) const
 {
 	IOnlineAccountIdRegistry* Registry = nullptr;
diff --git a/Engine/Source/Runtime/CoreOnline/Public/Online/CoreOnline.h b/Engine/Source/Runtime/CoreOnline/Public/Online/CoreOnline.h
--- a/Engine/Source/Runtime/CoreOnline/Public/Online/CoreOnline.h
+++ b/Engine/Source/Runtime/CoreOnline/Public/Online/CoreOnline.h
@@ -370,6 +370,7 @@ public:
 	virtual FString ToLogString(const TOnlineId<IdType>& OnlineId) const = 0;
 	virtual TArray<uint8> ToReplicationData(const TOnlineId<IdType>& OnlineId) const = 0;
 	virtual TOnlineId<IdType> FromReplicationData(const TArray<uint8>& ReplicationData) = 0;
+	virtual TOnlineId<IdType> FromStringData(const FString& StringData) = 0;
 };
 
 using IOnlineAccountIdRegistry = IOnlineIdRegistry<OnlineIdHandleTags::FAccount>;
@@ -413,6 +414,7 @@ public:
 	COREONLINE_API FString ToLogString(const FAccountId& AccountId) const;
 	COREONLINE_API TArray<uint8> ToReplicationData(const FAccountId& AccountId) const;
 	COREONLINE_API FAccountId ToAccountId(EOnlineServices Services, const TArray<uint8>& RepData) const;
+	COREONLINE_API FAccountId ToAccountId(EOnlineServices Services, const FString& StringData) const;
 
 	COREONLINE_API IOnlineAccountIdRegistry* GetAccountIdRegistry(EOnlineServices OnlineServices) const;
 
