OSSv2 Session interface
- Added aux GetSessionName
- Fixed typo in meta struct
- Added extra param (TargetAccountId) to remove session member async operation

#rnx

diff --git a/Engine/Plugins/Online/OnlineServices/Source/OnlineServicesCommon/Private/Online/SessionsCommon.cpp b/Engine/Plugins/Online/OnlineServices/Source/OnlineServicesCommon/Private/Online/SessionsCommon.cpp
--- a/Engine/Plugins/Online/OnlineServices/Source/OnlineServicesCommon/Private/Online/SessionsCommon.cpp
+++ b/Engine/Plugins/Online/OnlineServices/Source/OnlineServicesCommon/Private/Online/SessionsCommon.cpp
@@ -119,6 +119,7 @@ namespace UE::Online {
 
 		RegisterCommand(&FSessionsCommon::GetAllSessions);
 		RegisterCommand(&FSessionsCommon::GetSessionByName);
+		RegisterCommand(&FSessionsCommon::GetSessionName);
 		RegisterCommand(&FSessionsCommon::GetSessionById);
 		RegisterCommand(&FSessionsCommon::GetPresenceSession);
 		RegisterCommand(&FSessionsCommon::IsPresenceSession);
@@ -187,6 +188,19 @@ namespace UE::Online {
 		}	
 	}
 
+	TOnlineResult<FGetSessionName> FSessionsCommon::GetSessionName(FGetSessionName::Params&& Params) const
+	{
+		if (const FName* SessionName = LocalSessionsByName.FindKey(Params.SessionId))
+		{
+			return TOnlineResult<FGetSessionName>({ *SessionName });
+		}
+		else
+		{
+			return TOnlineResult<FGetSessionName>(Errors::NotFound());
+		}
+	}
+	
+
 	TOnlineResult<FGetSessionById> FSessionsCommon::GetSessionById(FGetSessionById::Params&& Params) const
 	{
 		TOnlineResult<FGetMutableSessionById> GetMutableSessionByIdResult = GetMutableSessionById({ Params.SessionId });
diff --git a/Engine/Plugins/Online/OnlineServices/Source/OnlineServicesCommon/Public/Online/SessionsCommon.h b/Engine/Plugins/Online/OnlineServices/Source/OnlineServicesCommon/Public/Online/SessionsCommon.h
--- a/Engine/Plugins/Online/OnlineServices/Source/OnlineServicesCommon/Public/Online/SessionsCommon.h
+++ b/Engine/Plugins/Online/OnlineServices/Source/OnlineServicesCommon/Public/Online/SessionsCommon.h
@@ -264,6 +264,7 @@ public:
 	// ISessions
 	UE_API virtual TOnlineResult<FGetAllSessions> GetAllSessions(FGetAllSessions::Params&& Params) const override;
 	UE_API virtual TOnlineResult<FGetSessionByName> GetSessionByName(FGetSessionByName::Params&& Params) const override;
+	UE_API virtual TOnlineResult<FGetSessionName> GetSessionName(FGetSessionName::Params&& Params) const override;
 	UE_API virtual TOnlineResult<FGetSessionById> GetSessionById(FGetSessionById::Params&& Params) const override;
 	UE_API virtual TOnlineResult<FGetPresenceSession> GetPresenceSession(FGetPresenceSession::Params&& Params) const override;
 	UE_API virtual TOnlineResult<FIsPresenceSession> IsPresenceSession(FIsPresenceSession::Params&& Params) const override;
diff --git a/Engine/Plugins/Online/OnlineServices/Source/OnlineServicesInterface/Public/Online/Sessions.h b/Engine/Plugins/Online/OnlineServices/Source/OnlineServicesInterface/Public/Online/Sessions.h
--- a/Engine/Plugins/Online/OnlineServices/Source/OnlineServicesInterface/Public/Online/Sessions.h
+++ b/Engine/Plugins/Online/OnlineServices/Source/OnlineServicesInterface/Public/Online/Sessions.h
@@ -243,6 +243,23 @@ struct FGetSessionByName
 	};
 };
 
+struct FGetSessionName
+{
+	static constexpr TCHAR Name[] = TEXT("GetSessionName");
+
+	struct Params
+	{
+		/** Id handle for the session to be retrieved the name from */
+		FOnlineSessionId SessionId;
+	};
+
+	struct Result
+	{
+		/** Session name */
+		FName SessionName;
+	};
+};
+
 struct FGetSessionById
 {
 	static constexpr TCHAR Name[] = TEXT("GetSessionById");
@@ -514,6 +531,9 @@ struct FRemoveSessionMember
 		/** Id handle for the local user which will perform the action */
 		FAccountId LocalAccountId;
 
+		/** Id handle for the user to be kicked */
+		FAccountId TargetAccountId; // TODO: Can we consider adding this parameter ?
+
 		/** Local name for the session */
 		FName SessionName;
 	};
@@ -704,6 +724,14 @@ public:
 	 */
 	virtual TOnlineResult<FGetSessionByName> GetSessionByName(FGetSessionByName::Params&& Params) const = 0;
 
+	/**
+	 * Gets the name of the session with the given id handle.
+	 *
+	 * @params Parameters for the GetSessionName call
+	 * return
+	 */
+	virtual TOnlineResult<FGetSessionName> GetSessionName(FGetSessionName::Params&& Params) const = 0;
+
 	/**
 	 * Gets a reference to the session with the given id handle.
 	 *
@@ -957,6 +985,15 @@ BEGIN_ONLINE_STRUCT_META(FGetSessionById::Result)
 	ONLINE_STRUCT_FIELD(FGetSessionById::Result, Session)
 END_ONLINE_STRUCT_META()
 
+BEGIN_ONLINE_STRUCT_META(FGetSessionName::Params)
+	ONLINE_STRUCT_FIELD(FGetSessionName::Params, SessionId)
+END_ONLINE_STRUCT_META()
+
+BEGIN_ONLINE_STRUCT_META(FGetSessionName::Result)
+	ONLINE_STRUCT_FIELD(FGetSessionName::Result, SessionName)
+END_ONLINE_STRUCT_META()
+
+
 BEGIN_ONLINE_STRUCT_META(FGetPresenceSession::Params)
 	ONLINE_STRUCT_FIELD(FGetPresenceSession::Params, LocalAccountId)
 END_ONLINE_STRUCT_META()
@@ -1062,6 +1099,7 @@ END_ONLINE_STRUCT_META()
 
 BEGIN_ONLINE_STRUCT_META(FRemoveSessionMember::Params)
 	ONLINE_STRUCT_FIELD(FRemoveSessionMember::Params, LocalAccountId),
+	ONLINE_STRUCT_FIELD(FRemoveSessionMember::Params, TargetAccountId),
 	ONLINE_STRUCT_FIELD(FRemoveSessionMember::Params, SessionName)
 END_ONLINE_STRUCT_META()
 
@@ -1125,8 +1163,8 @@ BEGIN_ONLINE_STRUCT_META(FSessionCreated)
 END_ONLINE_STRUCT_META()
 
 BEGIN_ONLINE_STRUCT_META(FSessionJoined)
-ONLINE_STRUCT_FIELD(FSessionCreated, LocalAccountId),
-ONLINE_STRUCT_FIELD(FSessionCreated, SessionId)
+ONLINE_STRUCT_FIELD(FSessionJoined, LocalAccountId),
+ONLINE_STRUCT_FIELD(FSessionJoined, SessionId)
 END_ONLINE_STRUCT_META()
 
 BEGIN_ONLINE_STRUCT_META(FSessionLeft)
