OSSv2 : Adding PSN Sessions interface methods

Also adding a flag in the Session's base for players who join sessions as spectators.

diff --git a/Engine/Plugins/Online/OnlineServices/Source/OnlineServicesInterface/Public/Online/OnlineErrorDefinitions.h b/Engine/Plugins/Online/OnlineServices/Source/OnlineServicesInterface/Public/Online/OnlineErrorDefinitions.h
--- a/Engine/Plugins/Online/OnlineServices/Source/OnlineServicesInterface/Public/Online/OnlineErrorDefinitions.h
+++ b/Engine/Plugins/Online/OnlineServices/Source/OnlineServicesInterface/Public/Online/OnlineErrorDefinitions.h
@@ -36,7 +36,8 @@ UE_ONLINE_ERROR_COMMON(Common, Unknown, 22, TEXT("unknown"), LOCTEXT("Unknown",
 UE_ONLINE_ERROR_COMMON(Common, NoChange, 23, TEXT("no_change"), LOCTEXT("NoChange", "No change"))
 UE_ONLINE_ERROR_COMMON(Common, AlreadyInSession, 24, TEXT("already_in_session"), LOCTEXT("AlreadyInSession", "Player is already in session"))
 UE_ONLINE_ERROR_COMMON(Common, NotInSession, 25, TEXT("not_in_session"), LOCTEXT("NotInSession", "Player is not in any session"))
-
+UE_ONLINE_ERROR_COMMON(Common, SessionJoinDenied, 26, TEXT("session_join_denied"), LOCTEXT("SessionJoinDenied", "Server denied session join"))
+UE_ONLINE_ERROR_COMMON(Common, SessionFull, 27, TEXT("session_full"), LOCTEXT("SessionFull", "Server session is full"))
 } /* namespace UE::Online::Errors */
 
 #undef LOCTEXT_NAMESPACE
diff --git a/Engine/Plugins/Online/OnlineServices/Source/OnlineServicesInterface/Public/Online/Sessions.h b/Engine/Plugins/Online/OnlineServices/Source/OnlineServicesInterface/Public/Online/Sessions.h
--- a/Engine/Plugins/Online/OnlineServices/Source/OnlineServicesInterface/Public/Online/Sessions.h
+++ b/Engine/Plugins/Online/OnlineServices/Source/OnlineServicesInterface/Public/Online/Sessions.h
@@ -416,7 +416,7 @@ struct FLeaveSession
 		FAccountId LocalAccountId;
 
 		/** Local name for the session */
-		FName SessionName;
+		FName SessionName; // TODO (MATIAS): This could be a FOnlineSessionId 
 
 		/** Whether the call should attempt to destroy the session instead of just leave it */
 		bool bDestroySession;
@@ -495,6 +495,9 @@ struct FJoinSession
 
 		/** Whether this session should be set as the user's new presence session. False by default */
 		bool bPresenceEnabled = false;
+
+		/** Whether to join as a spectator instead of player */
+		bool bJoinAsSpectator = false;
 	};
 
 	struct Result
@@ -533,7 +536,6 @@ struct FRemoveSessionMember
 
 		/** Id handle for the user to be kicked */
 		FAccountId TargetAccountId; // TODO: Can we consider adding this parameter ?
-
 		/** Local name for the session */
 		FName SessionName;
 	};
@@ -1083,7 +1085,8 @@ BEGIN_ONLINE_STRUCT_META(FJoinSession::Params)
 	ONLINE_STRUCT_FIELD(FJoinSession::Params, LocalAccountId),
 	ONLINE_STRUCT_FIELD(FJoinSession::Params, SessionName),
 	ONLINE_STRUCT_FIELD(FJoinSession::Params, SessionId),
-	ONLINE_STRUCT_FIELD(FJoinSession::Params, bPresenceEnabled)
+	ONLINE_STRUCT_FIELD(FJoinSession::Params, bPresenceEnabled),
+	ONLINE_STRUCT_FIELD(FJoinSession::Params, bJoinAsSpectator)
 END_ONLINE_STRUCT_META()
 
 BEGIN_ONLINE_STRUCT_META(FJoinSession::Result)
