Added logic and new parameter to Enqueue to optionally delay online async operation start.
Added usage of new logic in PresenceEOS.

diff --git a/Engine/Plugins/Online/OnlineServices/Source/OnlineServicesCommon/Private/Online/PresenceCommon.cpp b/Engine/Plugins/Online/OnlineServices/Source/OnlineServicesCommon/Private/Online/PresenceCommon.cpp
--- a/Engine/Plugins/Online/OnlineServices/Source/OnlineServicesCommon/Private/Online/PresenceCommon.cpp
+++ b/Engine/Plugins/Online/OnlineServices/Source/OnlineServicesCommon/Private/Online/PresenceCommon.cpp
@@ -20,6 +20,11 @@ void FPresenceCommon::RegisterCommands()
 	RegisterCommand(&FPresenceCommon::PartialUpdatePresence);
 }
 
+void FPresenceCommon::UpdateConfig()
+{
+	TOnlineComponent::UpdateConfig();
+}
+
 TOnlineAsyncOpHandle<FQueryPresence> FPresenceCommon::QueryPresence(FQueryPresence::Params&& Params)
 {
 	TOnlineAsyncOpRef<FQueryPresence> Operation = GetOp<FQueryPresence>(MoveTemp(Params));
diff --git a/Engine/Plugins/Online/OnlineServices/Source/OnlineServicesCommon/Public/Online/OnlineAsyncOp.h b/Engine/Plugins/Online/OnlineServices/Source/OnlineServicesCommon/Public/Online/OnlineAsyncOp.h
--- a/Engine/Plugins/Online/OnlineServices/Source/OnlineServicesCommon/Public/Online/OnlineAsyncOp.h
+++ b/Engine/Plugins/Online/OnlineServices/Source/OnlineServicesCommon/Public/Online/OnlineAsyncOp.h
@@ -1289,10 +1289,10 @@ public:
 	}
 
 	template <typename QueueType>
-	void Enqueue(QueueType& Queue)
+	void Enqueue(QueueType& Queue, double DelayInSeconds = 0.0)
 	{
 		static_assert(std::is_same_v<T, void>, "Continuation result discarded. Continuation prior to calling Enqueue must have a void or TFuture<void> return type.");
-		OwningOperation.Enqueue(Queue);
+		OwningOperation.Enqueue(Queue, DelayInSeconds);
 	}
 
 	TOnlineAsyncOp<OpType>& GetOwningOperation()
@@ -1329,9 +1329,9 @@ public:
 	}
 
 	template <typename QueueType>
-	void Enqueue(QueueType& Queue)
+	void Enqueue(QueueType& Queue, double DelayInSeconds = 0.0)
 	{
-		OwningOperation.Enqueue(Queue);
+		OwningOperation.Enqueue(Queue, DelayInSeconds);
 	}
 
 	TOnlineAsyncOp<OpType>& GetOwningOperation()
@@ -1462,12 +1462,12 @@ public:
 	IOnlineServices& GetOnlineServices() const { return Services; }
 
 	template <typename QueueType>
-	void Enqueue(QueueType& Queue)
+	void Enqueue(QueueType& Queue, double DelayInSeconds = 0.0)
 	{
 		check(SharedState->State < EAsyncOpState::Queued);
 		SharedState->State = EAsyncOpState::Queued;
 		UE_LOG(LogOnlineServices, Verbose, TEXT("%p %s op state set to: %s"), this, OpType::Name, LexToString(SharedState->State));
-		Queue.Enqueue(*this);
+		Queue.Enqueue(*this, DelayInSeconds);
 	}
 
 	void Start()
diff --git a/Engine/Plugins/Online/OnlineServices/Source/OnlineServicesCommon/Public/Online/OnlineAsyncOpQueue.h b/Engine/Plugins/Online/OnlineServices/Source/OnlineServicesCommon/Public/Online/OnlineAsyncOpQueue.h
--- a/Engine/Plugins/Online/OnlineServices/Source/OnlineServicesCommon/Public/Online/OnlineAsyncOpQueue.h
+++ b/Engine/Plugins/Online/OnlineServices/Source/OnlineServicesCommon/Public/Online/OnlineAsyncOpQueue.h
@@ -26,9 +26,9 @@ public:
 	}
 
 	template <typename OpType>
-	void Enqueue(TOnlineAsyncOp<OpType>& Op)
+	void Enqueue(TOnlineAsyncOp<OpType>& Op, double DelayInSeconds = 0.0)
 	{
-		QueuedOperations.Enqueue(MakeShared<TWrappedOp<OpType>>(Op));
+		QueuedOperations.Enqueue(MakeShared<TWrappedOp<OpType>>(Op, DelayInSeconds));
 		TryStartOperations();
 	}
 
@@ -51,8 +51,10 @@ protected:
 	public:
 		virtual ~TWrappedOp() {}
 
-		TWrappedOp(TOnlineAsyncOp<OpType>& InOp)
+		TWrappedOp(TOnlineAsyncOp<OpType>& InOp, double InDelayInSeconds)
 			: Op(InOp.AsShared())
+			, DelayInSeconds(InDelayInSeconds)
+			, EnqueueCycles(InDelayInSeconds > 0.0 ? FPlatformTime::Cycles64() : 0)
 		{
 		}
 
@@ -66,7 +68,31 @@ protected:
 					OnCompleteEvent.Broadcast();
 				}
 			});
-			Op->Start();
+			if (DelayInSeconds > 0.0)
+			{
+				const double ElapsedTime = FPlatformTime::ToSeconds64(FPlatformTime::Cycles64() - EnqueueCycles);
+				if (DelayInSeconds > ElapsedTime)
+				{
+					FTSTicker::GetCoreTicker().AddTicker(TEXT("StartDelayedOnlineAsyncOp"), DelayInSeconds - ElapsedTime,
+						[WeakOp = Op.ToWeakPtr()](float)
+						{
+							TSharedPtr<TOnlineAsyncOp<OpType>> PinnedOp = WeakOp.Pin();
+							if (PinnedOp.IsValid())
+							{
+								PinnedOp->Start();
+							}
+							return false;
+						});
+				}
+				else
+				{
+					Op->Start();
+				}
+			}
+			else
+			{
+				Op->Start();
+			}
 		}
 
 		virtual void Cancel(const FOnlineError& Reason) override final
@@ -86,6 +112,8 @@ protected:
 
 	private:
 		TSharedRef<TOnlineAsyncOp<OpType>> Op;
+		const double DelayInSeconds = 0.0;
+		const uint64 EnqueueCycles = 0;
 		TOnlineEventCallable<void(void)> OnCompleteEvent;
 		FOnlineEventDelegateHandle OnCompleteHandle;
 	};
diff --git a/Engine/Plugins/Online/OnlineServices/Source/OnlineServicesCommon/Public/Online/PresenceCommon.h b/Engine/Plugins/Online/OnlineServices/Source/OnlineServicesCommon/Public/Online/PresenceCommon.h
--- a/Engine/Plugins/Online/OnlineServices/Source/OnlineServicesCommon/Public/Online/PresenceCommon.h
+++ b/Engine/Plugins/Online/OnlineServices/Source/OnlineServicesCommon/Public/Online/PresenceCommon.h
@@ -20,6 +20,7 @@ public:
 
 	// IOnlineComponent
 	UE_API virtual void RegisterCommands() override;
+	UE_API virtual void UpdateConfig() override;
 
 	// IPresence
 	UE_API virtual TOnlineAsyncOpHandle<FQueryPresence> QueryPresence(FQueryPresence::Params&& Params) override;
diff --git a/Engine/Plugins/Online/OnlineServicesEOS/Source/Private/Online/PresenceEOS.cpp b/Engine/Plugins/Online/OnlineServicesEOS/Source/Private/Online/PresenceEOS.cpp
--- a/Engine/Plugins/Online/OnlineServicesEOS/Source/Private/Online/PresenceEOS.cpp
+++ b/Engine/Plugins/Online/OnlineServicesEOS/Source/Private/Online/PresenceEOS.cpp
@@ -15,6 +15,17 @@
 
 namespace UE::Online {
 
+struct FPresenceEOSConfig
+{
+	double AsyncOpEnqueueDelay = 0.0;
+};
+
+namespace Meta {
+	BEGIN_ONLINE_STRUCT_META(FPresenceEOSConfig)
+		ONLINE_STRUCT_FIELD(FPresenceEOSConfig, AsyncOpEnqueueDelay)
+	END_ONLINE_STRUCT_META()
+}
+
 static inline EUserPresenceStatus ToEPresenceState(EOS_Presence_EStatus InStatus)
 {
 	switch (InStatus)
@@ -96,6 +107,16 @@ void FPresenceEOS::PreShutdown()
 	OnPresenceChanged = nullptr;
 }
 
+void FPresenceEOS::UpdateConfig()
+{
+	Super::UpdateConfig();
+
+	FPresenceEOSConfig Config;
+	LoadConfig(Config);
+
+	AsyncOpEnqueueDelay = Config.AsyncOpEnqueueDelay;
+}
+
 void FPresenceEOS::HandleAuthLoginStatusChanged(const FAuthLoginStatusChanged& EventParameters)
 {
 	if (EventParameters.LoginStatus == ELoginStatus::LoggedIn)
@@ -178,7 +199,7 @@ TOnlineAsyncOpHandle<FQueryPresence> FPresenceEOS::QueryPresence(FQueryPresence:
 					InAsyncOp.SetError(Errors::FromEOSResult(Data->ResultCode));
 				}
 			})
-			.Enqueue(GetSerialQueue());
+			.Enqueue(GetSerialQueue(), AsyncOpEnqueueDelay);
 	}
 	return Op->GetHandle();
 }
@@ -483,7 +504,7 @@ TOnlineAsyncOpHandle<FUpdatePresence> FPresenceEOS::UpdatePresence(FUpdatePresen
 			InAsyncOp.SetError(Errors::FromEOSResult(Data->ResultCode));
 		}
 	})
-	.Enqueue(GetSerialQueue());
+	.Enqueue(GetSerialQueue(), AsyncOpEnqueueDelay);
 
 	return Op->GetHandle();
 }
@@ -758,7 +779,7 @@ TOnlineAsyncOpHandle<FPartialUpdatePresence> FPresenceEOS::PartialUpdatePresence
 					InAsyncOp.SetError(Errors::FromEOSResult(Data->ResultCode));
 				}
 			})
-			.Enqueue(GetSerialQueue());
+			.Enqueue(GetSerialQueue(), AsyncOpEnqueueDelay);
 	}
 	return Op->GetHandle();
 }
diff --git a/Engine/Plugins/Online/OnlineServicesEOS/Source/Public/Online/PresenceEOS.h b/Engine/Plugins/Online/OnlineServicesEOS/Source/Public/Online/PresenceEOS.h
--- a/Engine/Plugins/Online/OnlineServicesEOS/Source/Public/Online/PresenceEOS.h
+++ b/Engine/Plugins/Online/OnlineServicesEOS/Source/Public/Online/PresenceEOS.h
@@ -14,7 +14,7 @@
 namespace UE::Online {
 
 class FOnlineServicesEpicCommon;
-	
+
 class FPresenceEOS : public FPresenceCommon
 {
 public:
@@ -24,6 +24,7 @@ public:
 
 	ONLINESERVICESEOS_API virtual void Initialize() override;
 	ONLINESERVICESEOS_API virtual void PreShutdown() override;
+	ONLINESERVICESEOS_API virtual void UpdateConfig() override;
 
 	ONLINESERVICESEOS_API virtual TOnlineAsyncOpHandle<FQueryPresence> QueryPresence(FQueryPresence::Params&& Params) override;
 	ONLINESERVICESEOS_API virtual TOnlineResult<FGetCachedPresence> GetCachedPresence(FGetCachedPresence::Params&& Params) override;
@@ -56,6 +57,8 @@ protected:
 	TMap<FAccountId, TMap<FAccountId, TSharedRef<FUserPresence>>> PresenceLists;
 	
 	FEOSEventRegistrationPtr OnPresenceChanged;
+
+	double AsyncOpEnqueueDelay = 0.0;
 };
 
 /* UE::Online */ }
