Fix for Horde issue 1036238. Added missing includes.
Added cancel logic for latest timer changes.

diff --git a/Engine/Plugins/Online/OnlineServices/Source/OnlineServicesCommon/Public/Online/OnlineAsyncOpQueue.h b/Engine/Plugins/Online/OnlineServices/Source/OnlineServicesCommon/Public/Online/OnlineAsyncOpQueue.h
--- a/Engine/Plugins/Online/OnlineServices/Source/OnlineServicesCommon/Public/Online/OnlineAsyncOpQueue.h
+++ b/Engine/Plugins/Online/OnlineServices/Source/OnlineServicesCommon/Public/Online/OnlineAsyncOpQueue.h
@@ -4,6 +4,8 @@
 
 
 #include "Containers/SpscQueue.h"
+#include "Containers/Ticker.h"
+#include "HAL/PlatformTime.h"
 #include "Online/OnlineAsyncOpHandle.h"
 
 namespace UE::Online { template <typename OpType> class TOnlineAsyncOp; }
@@ -73,7 +75,7 @@ protected:
 				const double ElapsedTime = FPlatformTime::ToSeconds64(FPlatformTime::Cycles64() - EnqueueCycles);
 				if (DelayInSeconds > ElapsedTime)
 				{
-					FTSTicker::GetCoreTicker().AddTicker(TEXT("StartDelayedOnlineAsyncOp"), DelayInSeconds - ElapsedTime,
+					TickerHandle = FTSTicker::GetCoreTicker().AddTicker(TEXT("StartDelayedOnlineAsyncOp"), DelayInSeconds - ElapsedTime,
 						[WeakOp = Op.ToWeakPtr()](float)
 						{
 							TSharedPtr<TOnlineAsyncOp<OpType>> PinnedOp = WeakOp.Pin();
@@ -98,6 +100,11 @@ protected:
 		virtual void Cancel(const FOnlineError& Reason) override final
 		{
 			Op->Cancel(Reason);
+
+			if (TickerHandle.IsValid())
+			{
+				FTSTicker::GetCoreTicker().RemoveTicker(TickerHandle);
+			}
 		}
 
 		virtual bool IsComplete() const override final
@@ -116,6 +123,7 @@ protected:
 		const uint64 EnqueueCycles = 0;
 		TOnlineEventCallable<void(void)> OnCompleteEvent;
 		FOnlineEventDelegateHandle OnCompleteHandle;
+		FTSTicker::FDelegateHandle TickerHandle;
 	};
 
 	TSpscQueue<TSharedRef<IWrappedOp>> QueuedOperations;
