Use a TTransactionallySafeSpscQueue

#rnx

diff --git a/Engine/Plugins/Online/OnlineServices/Source/OnlineServicesCommon/Public/Online/OnlineAsyncOpQueue.h b/Engine/Plugins/Online/OnlineServices/Source/OnlineServicesCommon/Public/Online/OnlineAsyncOpQueue.h
--- a/Engine/Plugins/Online/OnlineServices/Source/OnlineServicesCommon/Public/Online/OnlineAsyncOpQueue.h
+++ b/Engine/Plugins/Online/OnlineServices/Source/OnlineServicesCommon/Public/Online/OnlineAsyncOpQueue.h
@@ -3,8 +3,8 @@
 #pragma once
 
 
-#include "Containers/SpscQueue.h"
 #include "Containers/Ticker.h"
+#include "Containers/TransactionallySafeSpscQueue.h"
 #include "HAL/PlatformTime.h"
 #include "Online/OnlineAsyncOpHandle.h"
 
@@ -126,7 +126,7 @@ protected:
 		FTSTicker::FDelegateHandle TickerHandle;
 	};
 
-	TSpscQueue<TSharedRef<IWrappedOp>> QueuedOperations;
+	TTransactionallySafeSpscQueue<TSharedRef<IWrappedOp>> QueuedOperations;
 };
 
 class FOnlineAsyncOpQueueParallel : public FOnlineAsyncOpQueue
