Fix for Horde Android packaging error.

diff --git a/Engine/Plugins/Online/OnlineServices/Source/OnlineServicesCommon/Public/Online/OnlineAsyncOpQueue.h b/Engine/Plugins/Online/OnlineServices/Source/OnlineServicesCommon/Public/Online/OnlineAsyncOpQueue.h
--- a/Engine/Plugins/Online/OnlineServices/Source/OnlineServicesCommon/Public/Online/OnlineAsyncOpQueue.h
+++ b/Engine/Plugins/Online/OnlineServices/Source/OnlineServicesCommon/Public/Online/OnlineAsyncOpQueue.h
@@ -103,7 +103,7 @@ protected:
 
 			if (TickerHandle.IsValid())
 			{
-				FTSTicker::GetCoreTicker().RemoveTicker(TickerHandle);
+				FTSTicker::RemoveTicker(TickerHandle);
 			}
 		}
 
