- Expand the set of SectionHierarchies we read config from

#rnx

diff --git a/Engine/Plugins/Online/OnlineServices/Source/OnlineServicesCommon/Public/Online/OnlineServicesCommon.h b/Engine/Plugins/Online/OnlineServices/Source/OnlineServicesCommon/Public/Online/OnlineServicesCommon.h
--- a/Engine/Plugins/Online/OnlineServices/Source/OnlineServicesCommon/Public/Online/OnlineServicesCommon.h
+++ b/Engine/Plugins/Online/OnlineServices/Source/OnlineServicesCommon/Public/Online/OnlineServicesCommon.h
@@ -2,14 +2,13 @@
 
 #pragma once
 
-#include "Online/OnlineServices.h"
-
-#include "Online/OnlineComponentRegistry.h"
+#include "Algo/NoneOf.h"
+#include "Containers/Ticker.h"
 #include "Online/OnlineAsyncOpCache.h"
 #include "Online/OnlineAsyncOpQueue.h"
+#include "Online/OnlineComponentRegistry.h"
 #include "Online/OnlineConfig.h"
-
-#include "Containers/Ticker.h"
+#include "Online/OnlineServices.h"
 
 #define UE_API ONLINESERVICESCOMMON_API
 
@@ -197,12 +196,16 @@ public:
 
 	/**
 	 * Load a config struct for an interface + operation
-	 * Will load values from the following sections:
+	 * Will load values from the following sections, skipping all with any empty strings:
 	 *   OnlineServices
-	 *   OnlineServices.<InterfaceName> (if InterfaceName is set)
+	 *   OnlineServices.<InterfaceName>
+	 *   OnlineServices.<InterfaceName>.<OperationName>
 	 *   OnlineServices.<ServiceProvider>
-	 *   OnlineServices.<ServiceProvider>.<InterfaceName> (if InterfaceName is set)
-	 *   OnlineServices.<ServiceProvider>.<InterfaceName>.<OperationName> (if OperationName is set)
+	 *   OnlineServices.<ServiceProvider>.<InterfaceName>
+	 *   OnlineServices.<ServiceProvider>.<InterfaceName>.<OperationName>
+	 *   OnlineServices.<ServiceProvider>.<InstanceConfigName>
+	 *   OnlineServices.<ServiceProvider>.<InstanceConfigName>.<InterfaceName>
+	 *   OnlineServices.<ServiceProvider>.<InstanceConfigName>.<InterfaceName>.<OperationName>
 	 * 
 	 * @param Struct Struct to populate with values from config
 	 * @param InterfaceName Optional interface name to append to the config section name
@@ -214,24 +217,28 @@ public:
 	bool LoadConfig(StructType& Struct, const FString& InterfaceName = FString(), const FString& OperationName = FString()) const
 	{
 		TArray<FString> SectionHeirarchy;
-		FString SectionName = TEXT("OnlineServices");
-		SectionHeirarchy.Add(SectionName);
-		if (!InterfaceName.IsEmpty())
-		{
-			SectionHeirarchy.Add(SectionName + TEXT(".") + InterfaceName);
-		}
-		SectionName += TEXT(".") + GetServiceConfigName();
-		SectionHeirarchy.Add(SectionName);
-		if (!InterfaceName.IsEmpty())
+
+		auto PushSection = [&SectionHeirarchy](const TArray<const FStringView>& SectionNames)
 		{
-			SectionName += TEXT(".") + InterfaceName;
-			SectionHeirarchy.Add(SectionName);
-			if (!OperationName.IsEmpty())
+			if (Algo::NoneOf(SectionNames, UE_PROJECTION_MEMBER(FStringView, IsEmpty)))
 			{
-				SectionName += TEXT(".") + OperationName;
-				SectionHeirarchy.Add(SectionName);
+				SectionHeirarchy.Emplace(FString::Join(SectionNames, TEXT(".")));
 			}
-		}
+		};
+
+		const FString& BaseSection = TEXT("OnlineServices");
+		const FString& InstanceConfigNameStr = InstanceConfigName.ToString();
+
+		PushSection({ BaseSection });
+		PushSection({ BaseSection, InterfaceName });
+		PushSection({ BaseSection, InterfaceName, OperationName });
+		PushSection({ BaseSection, ServiceConfigName });
+		PushSection({ BaseSection, ServiceConfigName, InterfaceName });
+		PushSection({ BaseSection, ServiceConfigName, InterfaceName, OperationName });
+		PushSection({ BaseSection, ServiceConfigName, InstanceConfigNameStr });
+		PushSection({ BaseSection, ServiceConfigName, InstanceConfigNameStr, InterfaceName });
+		PushSection({ BaseSection, ServiceConfigName, InstanceConfigNameStr, InterfaceName, OperationName });
+
 		return LoadConfig(Struct, SectionHeirarchy);
 	}
 	
