OSSv2: Adding PSN Stats interface and changes in PS5 UDS and Achievements

diff --git a/Engine/Plugins/Online/OnlineServices/Source/OnlineServicesCommon/Private/Online/StatsCommon.cpp b/Engine/Plugins/Online/OnlineServices/Source/OnlineServicesCommon/Private/Online/StatsCommon.cpp
--- a/Engine/Plugins/Online/OnlineServices/Source/OnlineServicesCommon/Private/Online/StatsCommon.cpp
+++ b/Engine/Plugins/Online/OnlineServices/Source/OnlineServicesCommon/Private/Online/StatsCommon.cpp
@@ -2,7 +2,7 @@
 
 #include "Online/StatsCommon.h"
 
-#include "Online/OnlineUtils.h"
+#include "Online/AuthCommon.h"
 
 namespace UE::Online {
 
@@ -79,39 +79,112 @@ void FStatsCommon::RegisterCommands()
 
 TOnlineAsyncOpHandle<FUpdateStats> FStatsCommon::UpdateStats(FUpdateStats::Params&& Params)
 {
-	TOnlineAsyncOpRef<FUpdateStats> Operation = GetOp<FUpdateStats>(MoveTemp(Params));
-	Operation->SetError(Errors::NotImplemented());
-	return Operation->GetHandle();
+	return UpdateStats_Implementation(MoveTemp(Params))->GetHandle(); 
 }
 
 TOnlineAsyncOpHandle<FQueryStats> FStatsCommon::QueryStats(FQueryStats::Params&& Params)
 {
-	TOnlineAsyncOpRef<FQueryStats> Operation = GetOp<FQueryStats>(MoveTemp(Params));
-	Operation->SetError(Errors::NotImplemented());
-	return Operation->GetHandle();
+	TOnlineAsyncOpRef<FQueryStats> Op = GetJoinableOp<FQueryStats>(MoveTemp(Params));
+	FAuthCommon* AuthPtr = Services.Get<FAuthCommon>();
+	if (AuthPtr == nullptr || !AuthPtr->IsLoggedIn(Op->GetParams().LocalAccountId))
+	{
+		Op->SetError(Errors::InvalidUser());
+		return Op->GetHandle();
+	}
+
+	Op->Then([this](TOnlineAsyncOp<FQueryStats>& InAsyncOp)
+	{
+		FQueryStats::Result Result;
+
+		FUserStats* ExistingUserStats = CachedUsersStats.FindByPredicate(FFindUserStatsByAccountId(InAsyncOp.GetParams().TargetAccountId));
+		ReadStatsFromCache(ExistingUserStats, InAsyncOp.GetParams().StatNames, Result.Stats);
+
+		InAsyncOp.SetResult(MoveTemp(Result));
+	})
+	.Enqueue(GetSerialQueue());
+
+	return Op->GetHandle();
 }
 
 TOnlineAsyncOpHandle<FTriggerEvent> FStatsCommon::TriggerEvent(FTriggerEvent::Params&& Params)
 {
-	TOnlineAsyncOpRef<FTriggerEvent> Operation = GetOp<FTriggerEvent>(MoveTemp(Params));
-	Operation->SetError(Errors::NotImplemented());
-	return Operation->GetHandle();
+	FUserStats UserStats;
+	UserStats.AccountId = Params.LocalAccountId;
+	for (const TPair<FString, FSchemaVariant>& EventPair : Params.EventParams)
+	{
+		// generic mapping: convert the event parameter value into a Stat value
+		UserStats.Stats.Emplace(EventPair.Key, FStatValue(EventPair.Value));
+	}
+
+	return UpdateStats_Implementation({ Params.LocalAccountId, { MoveTemp(UserStats) } })->GetWrappedHandle<FTriggerEvent>(
+	[this](TOnlineResult<FUpdateStats> UpdateResult) -> TOnlineResult<FTriggerEvent>
+	{
+		if (UpdateResult.IsOk())
+		{
+			// TODO: additional processing here
+			return TOnlineResult<FTriggerEvent>();
+		}
+
+		return TOnlineResult<FTriggerEvent>(MoveTemp(UpdateResult.GetErrorValue()));
+	});
 }
 
 TOnlineAsyncOpHandle<FBatchQueryStats> FStatsCommon::BatchQueryStats(FBatchQueryStats::Params&& Params)
 {
-	TOnlineAsyncOpRef<FBatchQueryStats> Operation = GetOp<FBatchQueryStats>(MoveTemp(Params));
-	Operation->SetError(Errors::NotImplemented());
-	return Operation->GetHandle();
-}
+	TOnlineAsyncOpRef<FBatchQueryStats> Op = GetOp<FBatchQueryStats>(MoveTemp(Params));
+	FAuthCommon* AuthPtr = Services.Get<FAuthCommon>();
+	if (AuthPtr == nullptr || !AuthPtr->IsLoggedIn(Op->GetParams().LocalAccountId))
+	{
+		Op->SetError(Errors::InvalidUser());
+		return Op->GetHandle();
+	}
 
+	Op->Then([this](TOnlineAsyncOp<FBatchQueryStats>& InAsyncOp)
+	{
+		FBatchQueryStats::Result Result;
+
+		for (const FAccountId& TargetAccountId : InAsyncOp.GetParams().TargetAccountIds)
+		{
+			FUserStats& UserStats = Result.UsersStats.Emplace_GetRef();
+			UserStats.AccountId = TargetAccountId;
+			FUserStats* ExistingUserStats = CachedUsersStats.FindByPredicate(FFindUserStatsByAccountId(TargetAccountId));
+			ReadStatsFromCache(ExistingUserStats, InAsyncOp.GetParams().StatNames, UserStats.Stats);
+		}
+
+		InAsyncOp.SetResult(MoveTemp(Result));
+	})
+	.Enqueue(GetSerialQueue());
+
+	return Op->GetHandle();
+}
 
 #if !UE_BUILD_SHIPPING
 TOnlineAsyncOpHandle<FResetStats> FStatsCommon::ResetStats(FResetStats::Params&& Params)
 {
-	TOnlineAsyncOpRef<FResetStats> Operation = GetOp<FResetStats>(MoveTemp(Params));
-	Operation->SetError(Errors::NotImplemented());
-	return Operation->GetHandle();
+	TOnlineAsyncOpRef<FResetStats> Op = GetOp<FResetStats>(MoveTemp(Params));
+	FAuthCommon* AuthPtr = Services.Get<FAuthCommon>();
+	if (AuthPtr == nullptr || !AuthPtr->IsLoggedIn(Op->GetParams().LocalAccountId))
+	{
+		Op->SetError(Errors::InvalidUser());
+		return Op->GetHandle();
+	}
+
+	Op->Then([this](TOnlineAsyncOp<FResetStats>& InAsyncOp)
+	{
+		const uint32 Index = CachedUsersStats.IndexOfByPredicate(FFindUserStatsByAccountId(InAsyncOp.GetParams().LocalAccountId));
+		if (Index != INDEX_NONE)
+		{
+			CachedUsersStats.RemoveAt(Index);
+		}
+
+		FFindUserStatsByAccountId FindUserStatsByAccountId(InAsyncOp.GetParams().LocalAccountId);
+		CachedUsersStats.RemoveAll(FindUserStatsByAccountId);
+
+		InAsyncOp.SetResult({});
+	})
+	.Enqueue(GetSerialQueue());
+
+	return Op->GetHandle();
 }
 #endif // !UE_BUILD_SHIPPING
 
@@ -152,4 +225,115 @@ const FStatDefinition* FStatsCommon::GetStatDefinition(const FString& StatName)
 	return StatDefinitions.Find(StatName); 
 }
 
+TOnlineAsyncOpRef<FUpdateStats> FStatsCommon::UpdateStats_Implementation(FUpdateStats::Params&& Params)
+{
+	TOnlineAsyncOpRef<FUpdateStats> Op = GetOp<FUpdateStats>(MoveTemp(Params));
+	FAuthCommon* AuthPtr = Services.Get<FAuthCommon>();
+	if (AuthPtr == nullptr || !AuthPtr->IsLoggedIn(Op->GetParams().LocalAccountId))
+	{
+		Op->SetError(Errors::InvalidUser());
+		return Op;
+	}
+
+	Op->Then([this](TOnlineAsyncOp<FUpdateStats>& InAsyncOp) mutable
+	{
+		for (const FUserStats& UpdateUserStats : InAsyncOp.GetParams().UpdateUsersStats)
+		{
+			FUserStats* ExistingUserStats = CachedUsersStats.FindByPredicate(FFindUserStatsByAccountId(UpdateUserStats.AccountId));
+			if (!ExistingUserStats)
+			{
+				ExistingUserStats = &CachedUsersStats.Emplace_GetRef();
+				ExistingUserStats->AccountId = UpdateUserStats.AccountId;
+			}
+
+			TMap<FString, FStatValue>& UserStats = ExistingUserStats->Stats;
+			for (const TPair<FString, FSchemaVariant>& UpdateUserStatPair : UpdateUserStats.Stats)
+			{
+				if (const FStatDefinition* StatDefinition = GetStatDefinition(UpdateUserStatPair.Key))
+				{
+					if (FStatValue* StatValue = UserStats.Find(UpdateUserStatPair.Key))
+					{
+						switch (StatDefinition->ModifyMethod)
+						{
+						case EStatModifyMethod::Set:
+							*StatValue = UpdateUserStatPair.Value;
+							break;
+						case EStatModifyMethod::Sum:
+							if (StatValue->VariantData.IsType<double>())
+							{
+								StatValue->Set(StatValue->GetDouble() + UpdateUserStatPair.Value.GetDouble());
+							}
+							else if (StatValue->VariantData.IsType<int64>())
+							{
+								StatValue->Set(StatValue->GetInt64() + UpdateUserStatPair.Value.GetInt64());
+							}
+							break;
+						case EStatModifyMethod::Largest:
+							if (StatValue->VariantData.IsType<double>())
+							{
+								StatValue->Set(FMath::Max(StatValue->GetDouble(), UpdateUserStatPair.Value.GetDouble()));
+							}
+							else if (StatValue->VariantData.IsType<int64>())
+							{
+								StatValue->Set(FMath::Max(StatValue->GetInt64(), UpdateUserStatPair.Value.GetInt64()));
+							}
+							break;
+						case EStatModifyMethod::Smallest:
+							if (StatValue->VariantData.IsType<double>())
+							{
+								StatValue->Set(FMath::Min(StatValue->GetDouble(), UpdateUserStatPair.Value.GetDouble()));
+							}
+							else if (StatValue->VariantData.IsType<int64>())
+							{
+								StatValue->Set(FMath::Min(StatValue->GetInt64(), UpdateUserStatPair.Value.GetInt64()));
+							}
+							break;
+						}
+					}
+					else
+					{
+						UserStats.Emplace(UpdateUserStatPair.Key, UpdateUserStatPair.Value);
+					}
+				}
+			}
+
+			CacheUserStats(*ExistingUserStats);
+		}
+
+		InAsyncOp.SetResult({});
+		OnStatsUpdatedEvent.Broadcast(InAsyncOp.GetParams());
+	})
+	.Enqueue(GetSerialQueue());
+
+	return Op;
+}
+
+void FStatsCommon::ReadStatsFromCache(const FUserStats* ExistingUserStats, const TArray<FString>& StatNames, TMap<FString, FStatValue>& OutStats) const
+{
+	for (const FString& StatName : StatNames)
+	{
+		if (ExistingUserStats)
+		{
+			if (const FStatValue* StatValue = ExistingUserStats->Stats.Find(StatName))
+			{
+				OutStats.Add(StatName, *StatValue);
+				continue;
+			}
+		}
+
+		if (const FStatDefinition* StatDefinition = StatDefinitions.Find(StatName))
+		{
+			const FStatValue& StatValue = OutStats.FindOrAdd(StatName, StatDefinition->DefaultValue);
+			if (StatValue.GetType() == ESchemaAttributeType::None)
+			{
+				UE_LOG(LogOnlineServices, Warning, TEXT("Stat Value in config definitions %s has no valid type"), *StatName);
+			}
+		}
+		else
+		{
+			UE_LOG(LogOnlineServices, Warning, TEXT("Make sure to add the Config.StatDefinitions array"));
+		}
+	}
+}
+
 /* UE::Online */ }
diff --git a/Engine/Plugins/Online/OnlineServices/Source/OnlineServicesCommon/Public/Online/StatsCommon.h b/Engine/Plugins/Online/OnlineServices/Source/OnlineServicesCommon/Public/Online/StatsCommon.h
--- a/Engine/Plugins/Online/OnlineServices/Source/OnlineServicesCommon/Public/Online/StatsCommon.h
+++ b/Engine/Plugins/Online/OnlineServices/Source/OnlineServicesCommon/Public/Online/StatsCommon.h
@@ -33,7 +33,7 @@ struct FStatDefinition
 	/* How the stat will be modified */
 	EStatModifyMethod ModifyMethod = EStatModifyMethod::Set;
 	/* Store the default type and value */
-	FStatValue DefaultValue = FStatValue(int64(0));
+	FStatValue DefaultValue = FStatValue(static_cast<int64>(0));
 };
 
 struct FStatsCommonConfig
@@ -86,7 +86,14 @@ public:
 	UE_API virtual TOnlineAsyncOpHandle<FUpdateStats> UpdateStats(FUpdateStats::Params&& Params) override;
 	UE_API virtual TOnlineAsyncOpHandle<FQueryStats> QueryStats(FQueryStats::Params&& Params) override;
 	UE_API virtual TOnlineAsyncOpHandle<FBatchQueryStats> BatchQueryStats(FBatchQueryStats::Params&& Params) override;
+
+	// Example implementation that uses the event to stats mapping rules
+	// In this version we assume that FTriggerEvent::Params is a mapping of property names to values,
+	// Assumes that FTriggerEvent::Params is defined as TMap<FString, FSchemaVariant>
+	// and that each stat referenced in the event has a defined modify rule: Set, Sum, Largest, Smallest
+	// After applying the stat changes to the base class cache, we're also broadcasting an update through OnStatsUpdatedEvent.Broadcast().
 	UE_API virtual TOnlineAsyncOpHandle<FTriggerEvent> TriggerEvent(FTriggerEvent::Params&& Params) override;
+
 #if !UE_BUILD_SHIPPING
 	UE_API virtual TOnlineAsyncOpHandle<FResetStats> ResetStats(FResetStats::Params&& Params) override;
 #endif // !UE_BUILD_SHIPPING
@@ -96,7 +103,9 @@ public:
 	UE_API const FStatDefinition* GetStatDefinition(const FString& StatName) const;
 
 protected:
-	UE_API void CacheUserStats(const FUserStats& UserStats);
+	UE_API virtual TOnlineAsyncOpRef<FUpdateStats> UpdateStats_Implementation(FUpdateStats::Params&& Params);
+	UE_API virtual void ReadStatsFromCache(const FUserStats* ExistingUserStats, const TArray<FString>& StatNames, TMap<FString, FStatValue>& OutStats) const;
+	UE_API virtual void CacheUserStats(const FUserStats& UserStats);
 
 	TMap<FString, FStatDefinition> StatDefinitions;
 	TOnlineEventCallable<void(const FStatsUpdated& StatsUpdated)> OnStatsUpdatedEvent;
diff --git a/Engine/Plugins/Online/OnlineServices/Source/OnlineServicesInterface/Public/Online/Stats.h b/Engine/Plugins/Online/OnlineServices/Source/OnlineServicesInterface/Public/Online/Stats.h
--- a/Engine/Plugins/Online/OnlineServices/Source/OnlineServicesInterface/Public/Online/Stats.h
+++ b/Engine/Plugins/Online/OnlineServices/Source/OnlineServicesInterface/Public/Online/Stats.h
@@ -140,28 +140,28 @@ public:
 	/**
 	 * Upload/send stats to the backend/platform.
 	 *
-	 * @param Parameters for the UpdateStats call
+	 * @param Params for the UpdateStats call
 	 * @return
 	 */
 	virtual TOnlineAsyncOpHandle<FUpdateStats> UpdateStats(FUpdateStats::Params&& Params) = 0;
 	/**
 	 * Get stats of the specified user.
 	 *
-	 * @param Parameters for the QueryStats call
+	 * @param Params for the QueryStats call
 	 * @return
 	 */
 	virtual TOnlineAsyncOpHandle<FQueryStats> QueryStats(FQueryStats::Params&& Params) = 0;
 	/**
 	 * Get stats of the specified users.
 	 *
-	 * @param Parameters for the BatchQueryStats call
+	 * @param Params for the BatchQueryStats call
 	 * @return
 	 */
 	virtual TOnlineAsyncOpHandle<FBatchQueryStats> BatchQueryStats(FBatchQueryStats::Params&& Params) = 0;
 	/**
 	 * Trigger a stat modifying event
 	 *
-	 * @param Parameters for the TriggerEvent call
+	 * @param Params for the TriggerEvent call
 	 * @return
 	 */
 	virtual TOnlineAsyncOpHandle<FTriggerEvent> TriggerEvent(FTriggerEvent::Params&& Params) = 0;
@@ -170,7 +170,7 @@ public:
 	/**
 	 * Reset stats of the local user.
 	 *
-	 * @param Parameters for the ResetStats call
+	 * @param Params for the ResetStats call
 	 * @return
 	 */
 	virtual TOnlineAsyncOpHandle<FResetStats> ResetStats(FResetStats::Params&& Params) = 0;
diff --git a/Engine/Plugins/Online/OnlineServicesNull/Source/Private/Online/StatsNull.cpp b/Engine/Plugins/Online/OnlineServicesNull/Source/Private/Online/StatsNull.cpp
--- a/Engine/Plugins/Online/OnlineServicesNull/Source/Private/Online/StatsNull.cpp
+++ b/Engine/Plugins/Online/OnlineServicesNull/Source/Private/Online/StatsNull.cpp
@@ -1,7 +1,6 @@
 // Copyright Epic Games, Inc. All Rights Reserved.
 
 #include "Online/StatsNull.h"
-#include "Online/AuthNull.h"
 #include "Online/OnlineServicesNull.h"
 
 namespace UE::Online {
@@ -11,190 +10,4 @@ FStatsNull::FStatsNull(FOnlineServicesNull& InOwningSubsystem)
 {
 }
 
-TOnlineAsyncOpHandle<FUpdateStats> FStatsNull::UpdateStats(FUpdateStats::Params&& Params)
-{
-	TOnlineAsyncOpRef<FUpdateStats> Op = GetOp<FUpdateStats>(MoveTemp(Params));
-
-	if (!Services.Get<FAuthNull>()->IsLoggedIn(Op->GetParams().LocalAccountId))
-	{
-		Op->SetError(Errors::InvalidUser());
-		return Op->GetHandle();
-	}
-
-	Op->Then([this](TOnlineAsyncOp<FUpdateStats>& InAsyncOp) mutable
-	{
-		for (const FUserStats& UpdateUserStats : InAsyncOp.GetParams().UpdateUsersStats)
-		{
-			FUserStats* ExistingUserStats = UsersStats.FindByPredicate(FFindUserStatsByAccountId(UpdateUserStats.AccountId));
-			if (!ExistingUserStats)
-			{
-				ExistingUserStats = &UsersStats.Emplace_GetRef();
-				ExistingUserStats->AccountId = UpdateUserStats.AccountId;
-			}
-
-			TMap<FString, FStatValue>& UserStats = ExistingUserStats->Stats;
-			for (const auto& UpdateUserStatPair : UpdateUserStats.Stats)
-			{
-				if (const FStatDefinition* StatDefinition = GetStatDefinition(UpdateUserStatPair.Key))
-				{
-					if (FStatValue* StatValue = UserStats.Find(UpdateUserStatPair.Key))
-					{
-						switch (StatDefinition->ModifyMethod)
-						{
-						case EStatModifyMethod::Set:
-							*StatValue = UpdateUserStatPair.Value;
-							break;
-						case EStatModifyMethod::Sum:
-							if (StatValue->VariantData.IsType<double>())
-							{
-								StatValue->Set(StatValue->GetDouble() + UpdateUserStatPair.Value.GetDouble());
-							}
-							else if (StatValue->VariantData.IsType<int64>())
-							{
-								StatValue->Set(StatValue->GetInt64() + UpdateUserStatPair.Value.GetInt64());
-							}
-							break;
-						case EStatModifyMethod::Largest:
-							if (StatValue->VariantData.IsType<double>())
-							{
-								StatValue->Set(FMath::Max(StatValue->GetDouble(), UpdateUserStatPair.Value.GetDouble()));
-							}
-							else if (StatValue->VariantData.IsType<int64>())
-							{
-								StatValue->Set(FMath::Max(StatValue->GetInt64(), UpdateUserStatPair.Value.GetInt64()));
-							}
-							break;
-						case EStatModifyMethod::Smallest:
-							if (StatValue->VariantData.IsType<double>())
-							{
-								StatValue->Set(FMath::Min(StatValue->GetDouble(), UpdateUserStatPair.Value.GetDouble()));
-							}
-							else if (StatValue->VariantData.IsType<int64>())
-							{
-								StatValue->Set(FMath::Min(StatValue->GetInt64(), UpdateUserStatPair.Value.GetInt64()));
-							}
-							break;
-						}
-					}
-					else
-					{
-						UserStats.Emplace(UpdateUserStatPair.Key, UpdateUserStatPair.Value);
-					}
-				}
-			}
-
-			CacheUserStats(*ExistingUserStats);
-		}
-
-		InAsyncOp.SetResult({});
-		OnStatsUpdatedEvent.Broadcast(InAsyncOp.GetParams());
-	})
-	.Enqueue(GetSerialQueue());
-
-	return Op->GetHandle();
-}
-
-void FStatsNull::ReadStatsFromCache(const FUserStats* ExistingUserStats, const TArray<FString>& StatNames, TMap<FString, FStatValue>& OutStats)
-{
-	for (const FString& StatName : StatNames)
-	{
-		if (ExistingUserStats)
-		{
-			if (const FStatValue* StatValue = ExistingUserStats->Stats.Find(StatName))
-			{
-				OutStats.Add(StatName, *StatValue);
-				continue;
-			}
-		}
-
-		const FStatDefinition* StatDefinition = GetStatDefinition(StatName);
-		check(StatDefinition);
-		OutStats.Add(StatName, StatDefinition->DefaultValue);
-	}
-}
-
-TOnlineAsyncOpHandle<FQueryStats> FStatsNull::QueryStats(FQueryStats::Params&& Params)
-{
-	TOnlineAsyncOpRef<FQueryStats> Op = GetOp<FQueryStats>(MoveTemp(Params));
-
-	if (!Services.Get<FAuthNull>()->IsLoggedIn(Op->GetParams().LocalAccountId))
-	{
-		Op->SetError(Errors::InvalidUser());
-		return Op->GetHandle();
-	}
-
-	Op->Then([this](TOnlineAsyncOp<FQueryStats>& InAsyncOp)
-	{
-		FQueryStats::Result Result;
-
-		FUserStats* ExistingUserStats = UsersStats.FindByPredicate(FFindUserStatsByAccountId(InAsyncOp.GetParams().TargetAccountId));
-		ReadStatsFromCache(ExistingUserStats, InAsyncOp.GetParams().StatNames, Result.Stats);
-
-		InAsyncOp.SetResult(MoveTemp(Result));
-	})
-	.Enqueue(GetSerialQueue());
-
-	return Op->GetHandle();
-}
-
-TOnlineAsyncOpHandle<FBatchQueryStats> FStatsNull::BatchQueryStats(FBatchQueryStats::Params&& Params)
-{
-	TOnlineAsyncOpRef<FBatchQueryStats> Op = GetOp<FBatchQueryStats>(MoveTemp(Params));
-	if (!Services.Get<FAuthNull>()->IsLoggedIn(Op->GetParams().LocalAccountId))
-	{
-		Op->SetError(Errors::InvalidUser());
-		return Op->GetHandle();
-	}
-
-	Op->Then([this](TOnlineAsyncOp<FBatchQueryStats>& InAsyncOp)
-	{
-		FBatchQueryStats::Result Result;
-
-		for (const FAccountId& TargetAccountId : InAsyncOp.GetParams().TargetAccountIds)
-		{
-			FUserStats& UserStats = Result.UsersStats.Emplace_GetRef();
-			UserStats.AccountId = TargetAccountId;
-			FUserStats* ExistingUserStats = UsersStats.FindByPredicate(FFindUserStatsByAccountId(TargetAccountId));
-			ReadStatsFromCache(ExistingUserStats, InAsyncOp.GetParams().StatNames, UserStats.Stats);
-		}
-
-		InAsyncOp.SetResult(MoveTemp(Result));
-	})
-	.Enqueue(GetSerialQueue());
-
-	return Op->GetHandle();
-}
-
-#if !UE_BUILD_SHIPPING
-TOnlineAsyncOpHandle<FResetStats> FStatsNull::ResetStats(FResetStats::Params&& Params)
-{
-	TOnlineAsyncOpRef<FResetStats> Op = GetOp<FResetStats>(MoveTemp(Params));
-
-	if (!Services.Get<FAuthNull>()->IsLoggedIn(Op->GetParams().LocalAccountId))
-	{
-		Op->SetError(Errors::InvalidUser());
-		return Op->GetHandle();
-	}
-
-	Op->Then([this](TOnlineAsyncOp<FResetStats>& InAsyncOp)
-	{
-		uint32 Index = UsersStats.IndexOfByPredicate(FFindUserStatsByAccountId(InAsyncOp.GetParams().LocalAccountId));
-		if (Index != INDEX_NONE)
-		{
-			UsersStats.RemoveAt(Index);
-		}
-
-		FFindUserStatsByAccountId FindUserStatsByAccountId(InAsyncOp.GetParams().LocalAccountId);
-
-		UsersStats.RemoveAll(FindUserStatsByAccountId);
-		CachedUsersStats.RemoveAll(FindUserStatsByAccountId);
-
-		InAsyncOp.SetResult({});
-	})
-	.Enqueue(GetSerialQueue());
-
-	return Op->GetHandle();
-}
-#endif // !UE_BUILD_SHIPPING
-
 /* UE::Online */ }
diff --git a/Engine/Plugins/Online/OnlineServicesNull/Source/Public/Online/StatsNull.h b/Engine/Plugins/Online/OnlineServicesNull/Source/Public/Online/StatsNull.h
--- a/Engine/Plugins/Online/OnlineServicesNull/Source/Public/Online/StatsNull.h
+++ b/Engine/Plugins/Online/OnlineServicesNull/Source/Public/Online/StatsNull.h
@@ -2,7 +2,6 @@
 
 #pragma once
 
-#include "Misc/Optional.h"
 #include "Online/StatsCommon.h"
 
 #define UE_API ONLINESERVICESNULL_API
@@ -17,20 +16,6 @@ public:
 	using Super = FStatsCommon;
 
 	UE_API FStatsNull(FOnlineServicesNull& InOwningSubsystem);
-
-	// IStats
-	UE_API virtual TOnlineAsyncOpHandle<FUpdateStats> UpdateStats(FUpdateStats::Params&& Params) override;
-	UE_API virtual TOnlineAsyncOpHandle<FQueryStats> QueryStats(FQueryStats::Params&& Params) override;
-	UE_API virtual TOnlineAsyncOpHandle<FBatchQueryStats> BatchQueryStats(FBatchQueryStats::Params&& Params) override;
-#if !UE_BUILD_SHIPPING
-	UE_API virtual TOnlineAsyncOpHandle<FResetStats> ResetStats(FResetStats::Params&& Params) override;
-#endif // !UE_BUILD_SHIPPING
-
-protected:
-	UE_API void ReadStatsFromCache(const FUserStats* ExistingUserStats, const TArray<FString>& StatNames, TMap<FString, FStatValue>& OutStats);
-
-	// TODO: Save UsersStats into local user profile
-	TArray<FUserStats> UsersStats;
 };
 
 /* UE::Online */ }
