Add missed virtual dtor for OSSv2 interfaces.

diff --git a/Engine/Plugins/Online/OnlineServices/Source/OnlineServicesInterface/Public/Online/Achievements.h b/Engine/Plugins/Online/OnlineServices/Source/OnlineServicesInterface/Public/Online/Achievements.h
--- a/Engine/Plugins/Online/OnlineServices/Source/OnlineServicesInterface/Public/Online/Achievements.h
+++ b/Engine/Plugins/Online/OnlineServices/Source/OnlineServicesInterface/Public/Online/Achievements.h
@@ -194,6 +194,8 @@ bool ONLINESERVICESINTERFACE_API operator==(const FAchievementStateUpdated&, con
 class IAchievements
 {
 public:
+	virtual ~IAchievements() {}
+
 	/**
 	 * Query all achievement definitions.
 	 */
diff --git a/Engine/Plugins/Online/OnlineServices/Source/OnlineServicesInterface/Public/Online/Auth.h b/Engine/Plugins/Online/OnlineServices/Source/OnlineServicesInterface/Public/Online/Auth.h
--- a/Engine/Plugins/Online/OnlineServices/Source/OnlineServicesInterface/Public/Online/Auth.h
+++ b/Engine/Plugins/Online/OnlineServices/Source/OnlineServicesInterface/Public/Online/Auth.h
@@ -464,6 +464,8 @@ struct FAuthAccountAttributesChanged
 class IAuth
 {
 public:
+	virtual ~IAuth() {}
+
 	/**
 	 * Authenticate a local user. If necessary, retrieve and maintain an access token for the
 	 * duration of the auth session.
diff --git a/Engine/Plugins/Online/OnlineServices/Source/OnlineServicesInterface/Public/Online/Commerce.h b/Engine/Plugins/Online/OnlineServices/Source/OnlineServicesInterface/Public/Online/Commerce.h
--- a/Engine/Plugins/Online/OnlineServices/Source/OnlineServicesInterface/Public/Online/Commerce.h
+++ b/Engine/Plugins/Online/OnlineServices/Source/OnlineServicesInterface/Public/Online/Commerce.h
@@ -319,6 +319,7 @@ public:
 class ICommerce
 {
 public:
+	virtual ~ICommerce() {}
 
 	/* Query the list of all offers for the target user */
 	virtual TOnlineAsyncOpHandle<FCommerceQueryOffers> QueryOffers(FCommerceQueryOffers::Params&& Params) = 0;
diff --git a/Engine/Plugins/Online/OnlineServices/Source/OnlineServicesInterface/Public/Online/Connectivity.h b/Engine/Plugins/Online/OnlineServices/Source/OnlineServicesInterface/Public/Online/Connectivity.h
--- a/Engine/Plugins/Online/OnlineServices/Source/OnlineServicesInterface/Public/Online/Connectivity.h
+++ b/Engine/Plugins/Online/OnlineServices/Source/OnlineServicesInterface/Public/Online/Connectivity.h
@@ -55,6 +55,8 @@ struct FConnectionStatusChanged
 class IConnectivity
 {
 public:
+	virtual ~IConnectivity() {}
+
 	/**
 		* Accessor for the connection status
 		*/
diff --git a/Engine/Plugins/Online/OnlineServices/Source/OnlineServicesInterface/Public/Online/ExternalUI.h b/Engine/Plugins/Online/OnlineServices/Source/OnlineServicesInterface/Public/Online/ExternalUI.h
--- a/Engine/Plugins/Online/OnlineServices/Source/OnlineServicesInterface/Public/Online/ExternalUI.h
+++ b/Engine/Plugins/Online/OnlineServices/Source/OnlineServicesInterface/Public/Online/ExternalUI.h
@@ -58,6 +58,8 @@ struct FExternalUIStatusChanged
 class IExternalUI
 {
 public:
+	virtual ~IExternalUI() {}
+
 	/**
 	 * Shows the online service's login UI
 	 */
diff --git a/Engine/Plugins/Online/OnlineServices/Source/OnlineServicesInterface/Public/Online/Lobbies.h b/Engine/Plugins/Online/OnlineServices/Source/OnlineServicesInterface/Public/Online/Lobbies.h
--- a/Engine/Plugins/Online/OnlineServices/Source/OnlineServicesInterface/Public/Online/Lobbies.h
+++ b/Engine/Plugins/Online/OnlineServices/Source/OnlineServicesInterface/Public/Online/Lobbies.h
@@ -620,6 +620,8 @@ struct FUILobbyJoinRequested
 class ILobbies
 {
 public:
+	virtual ~ILobbies() {}
+
 	//----------------------------------------------------------------------------------------------
 	// Operations
 
diff --git a/Engine/Plugins/Online/OnlineServices/Source/OnlineServicesInterface/Public/Online/OnlineAsyncOpHandle.h b/Engine/Plugins/Online/OnlineServices/Source/OnlineServicesInterface/Public/Online/OnlineAsyncOpHandle.h
--- a/Engine/Plugins/Online/OnlineServices/Source/OnlineServicesInterface/Public/Online/OnlineAsyncOpHandle.h
+++ b/Engine/Plugins/Online/OnlineServices/Source/OnlineServicesInterface/Public/Online/OnlineAsyncOpHandle.h
@@ -384,6 +384,8 @@ inline const TCHAR* LexToString(EAsyncOpState Value)
 class IOnlineAsyncOp
 {
 public:
+	virtual ~IOnlineAsyncOp() {}
+
 	virtual void Cancel() = 0;
 	virtual EAsyncOpState GetState() const = 0;
 	virtual void DetachHandle(class IAsyncHandle* Handle) = 0;
diff --git a/Engine/Plugins/Online/OnlineServices/Source/OnlineServicesInterface/Public/Online/OnlineError.h b/Engine/Plugins/Online/OnlineServices/Source/OnlineServicesInterface/Public/Online/OnlineError.h
--- a/Engine/Plugins/Online/OnlineServices/Source/OnlineServicesInterface/Public/Online/OnlineError.h
+++ b/Engine/Plugins/Online/OnlineServices/Source/OnlineServicesInterface/Public/Online/OnlineError.h
@@ -14,6 +14,8 @@ class FOnlineError;
 class IOnlineErrorDetails
 {
 public:
+	virtual ~IOnlineErrorDetails() {}
+
 	virtual FString GetFriendlyErrorCode(const FOnlineError&) const = 0;
 	virtual FText GetText(const FOnlineError&) const = 0;
 	virtual FString GetLogString(const FOnlineError&) const = 0;
diff --git a/Engine/Plugins/Online/OnlineServices/Source/OnlineServicesInterface/Public/Online/Presence.h b/Engine/Plugins/Online/OnlineServices/Source/OnlineServicesInterface/Public/Online/Presence.h
--- a/Engine/Plugins/Online/OnlineServices/Source/OnlineServicesInterface/Public/Online/Presence.h
+++ b/Engine/Plugins/Online/OnlineServices/Source/OnlineServicesInterface/Public/Online/Presence.h
@@ -322,6 +322,8 @@ struct FPresenceUpdated
 class IPresence
 {
 public:
+	virtual ~IPresence() {}
+
 	/**
 	 * Query presence for a user
 	 * 
diff --git a/Engine/Plugins/Online/OnlineServices/Source/OnlineServicesInterface/Public/Online/Sessions.h b/Engine/Plugins/Online/OnlineServices/Source/OnlineServicesInterface/Public/Online/Sessions.h
--- a/Engine/Plugins/Online/OnlineServices/Source/OnlineServicesInterface/Public/Online/Sessions.h
+++ b/Engine/Plugins/Online/OnlineServices/Source/OnlineServicesInterface/Public/Online/Sessions.h
@@ -162,6 +162,8 @@ struct FSessionInfo
 class ISession
 {
 public:
+	virtual ~ISession() {}
+
 	/** Retrieves the id handle for the user who created or currently owns the session */
 	virtual const FAccountId GetOwnerAccountId() const = 0;
 
diff --git a/Engine/Plugins/Online/OnlineServices/Source/OnlineServicesInterface/Public/Online/Social.h b/Engine/Plugins/Online/OnlineServices/Source/OnlineServicesInterface/Public/Online/Social.h
--- a/Engine/Plugins/Online/OnlineServices/Source/OnlineServicesInterface/Public/Online/Social.h
+++ b/Engine/Plugins/Online/OnlineServices/Source/OnlineServicesInterface/Public/Online/Social.h
@@ -216,6 +216,8 @@ struct FBlockUser
 class ISocial
 {
 public:
+	virtual ~ISocial() {}
+
 	/**
 	 * Query the friends list
 	 * 
diff --git a/Engine/Plugins/Online/OnlineServices/Source/OnlineServicesInterface/Public/Online/Stats.h b/Engine/Plugins/Online/OnlineServices/Source/OnlineServicesInterface/Public/Online/Stats.h
--- a/Engine/Plugins/Online/OnlineServices/Source/OnlineServicesInterface/Public/Online/Stats.h
+++ b/Engine/Plugins/Online/OnlineServices/Source/OnlineServicesInterface/Public/Online/Stats.h
@@ -137,6 +137,8 @@ using FStatsUpdated = FUpdateStats::Params;
 class IStats
 {
 public:
+	virtual ~IStats() {}
+
 	/**
 	 * Upload/send stats to the backend/platform.
 	 *
diff --git a/Engine/Plugins/Online/OnlineServices/Source/OnlineServicesInterface/Public/Online/TitleFile.h b/Engine/Plugins/Online/OnlineServices/Source/OnlineServicesInterface/Public/Online/TitleFile.h
--- a/Engine/Plugins/Online/OnlineServices/Source/OnlineServicesInterface/Public/Online/TitleFile.h
+++ b/Engine/Plugins/Online/OnlineServices/Source/OnlineServicesInterface/Public/Online/TitleFile.h
@@ -70,6 +70,8 @@ struct FTitleFileReadFile
 class ITitleFile
 {
 public:
+	virtual ~ITitleFile() {}
+
 	/**
 	 * Enumerate all available files
 	 */
diff --git a/Engine/Plugins/Online/OnlineServices/Source/OnlineServicesInterface/Public/Online/UserFile.h b/Engine/Plugins/Online/OnlineServices/Source/OnlineServicesInterface/Public/Online/UserFile.h
--- a/Engine/Plugins/Online/OnlineServices/Source/OnlineServicesInterface/Public/Online/UserFile.h
+++ b/Engine/Plugins/Online/OnlineServices/Source/OnlineServicesInterface/Public/Online/UserFile.h
@@ -125,6 +125,8 @@ struct FUserFileDeleteFile
 class IUserFile
 {
 public:
+	virtual ~IUserFile() {}
+
 	/**
 	 * Enumerate all available files
 	 */
diff --git a/Engine/Plugins/Online/OnlineServices/Source/OnlineServicesInterface/Public/Online/UserInfo.h b/Engine/Plugins/Online/OnlineServices/Source/OnlineServicesInterface/Public/Online/UserInfo.h
--- a/Engine/Plugins/Online/OnlineServices/Source/OnlineServicesInterface/Public/Online/UserInfo.h
+++ b/Engine/Plugins/Online/OnlineServices/Source/OnlineServicesInterface/Public/Online/UserInfo.h
@@ -136,6 +136,8 @@ struct FShowUserProfile
 class IUserInfo
 {
 public:
+	virtual ~IUserInfo() {}
+
 	/*
 	 * Queries the user info for a list of account ids.
 	 * Upon success, user info can be accessed with GetUserInfo.
