OnlineServices: Move TOnlineAsyncOp::ExecuteStep logic to shared based class.

diff --git a/Engine/Plugins/Online/OnlineServices/Source/OnlineServicesCommon/Public/Online/OnlineAsyncOp.h b/Engine/Plugins/Online/OnlineServices/Source/OnlineServicesCommon/Public/Online/OnlineAsyncOp.h
--- a/Engine/Plugins/Online/OnlineServices/Source/OnlineServicesCommon/Public/Online/OnlineAsyncOp.h
+++ b/Engine/Plugins/Online/OnlineServices/Source/OnlineServicesCommon/Public/Online/OnlineAsyncOp.h
@@ -1353,6 +1353,86 @@ public:
 	virtual void SetError(FOnlineError&& Error) = 0;
 	virtual void Cancel(const FOnlineError& Reason) = 0;
 	virtual void ClearCallback() = 0;
+
+	void AddStep(TUniquePtr<Private::IStep>&& Step)
+	{
+		Steps.Add(MoveTemp(Step));
+	}
+
+protected:
+	template <typename CallableType>
+	void ExecuteImpl(const TCHAR* DebugName, FOnlineAsyncExecutionPolicy ExecutionPolicy, CallableType&& Callable)
+	{
+		switch (ExecutionPolicy.GetExecutionPolicy())
+		{
+		case EOnlineAsyncExecutionPolicy::RunOnGameThread:
+			if (IsInGameThread())
+			{
+				Callable();
+			}
+			else
+			{
+				ExecuteOnGameThread(DebugName, MoveTemp(Callable));
+			}
+			break;
+
+		case EOnlineAsyncExecutionPolicy::RunOnNextTick:
+			FTSTicker::GetCoreTicker().AddTicker(FTickerDelegate::CreateLambda([DebugName, LambdaCallable = Callable](float dt) mutable
+				{
+					ExecuteOnGameThread(DebugName, MoveTemp(LambdaCallable));
+					return false;
+				}));
+			break;
+
+		case EOnlineAsyncExecutionPolicy::RunOnThreadPool:
+			Async(EAsyncExecution::ThreadPool, MoveTemp(Callable));
+			break;
+
+		case EOnlineAsyncExecutionPolicy::RunOnTaskGraph:
+			Async(EAsyncExecution::TaskGraph, MoveTemp(Callable));
+			break;
+
+		case EOnlineAsyncExecutionPolicy::RunImmediately:
+			Callable();
+			break;
+		}
+	}
+
+	void ExecuteNextStepImpl(const TCHAR* DebugName, const TSharedRef<FOnlineAsyncOp>& SharedThis)
+	{
+		const int32 StepToExecute = NextStep;
+		++NextStep;
+		if (StepToExecute < Steps.Num())
+		{
+			ExecuteImpl(DebugName, Steps[StepToExecute]->GetExecutionPolicy(),
+				[this, StepToExecute, WeakThis = TWeakPtr<FOnlineAsyncOp>(SharedThis)]()
+				{
+					if (TSharedPtr<FOnlineAsyncOp> PinnedThis = WeakThis.Pin())
+					{
+						Steps[StepToExecute]->Execute();
+					}
+				});
+		}
+	}
+
+	void ExecuteRepeatStepImpl(const TCHAR* DebugName, const TSharedRef<FOnlineAsyncOp>& SharedThis)
+	{
+		const int32 StepToExecute = NextStep == 0 ? 0 : NextStep - 1;
+		if (StepToExecute < Steps.Num())
+		{
+			ExecuteImpl(DebugName, Steps[StepToExecute]->GetExecutionPolicy(),
+				[this, StepToExecute, WeakThis = TWeakPtr<FOnlineAsyncOp>(SharedThis)]()
+				{
+					if (TSharedPtr<FOnlineAsyncOp> PinnedThis = WeakThis.Pin())
+					{
+						Steps[StepToExecute]->Execute();
+					}
+				});
+		}
+	}
+
+	TArray<TUniquePtr<Private::IStep>> Steps;
+	int32 NextStep = 0;
 };
 
 // This class represents an async operation on the public interface
@@ -1482,19 +1562,7 @@ public:
 	{
 		if (!IsComplete())
 		{
-			const int32 StepToExecute = NextStep;
-			++NextStep;
-			if (StepToExecute < Steps.Num())
-			{
-				Execute(Steps[StepToExecute]->GetExecutionPolicy(),
-					[this, StepToExecute, WeakThis = TWeakPtr<TOnlineAsyncOp<OpType>>(this->AsShared())]()
-					{
-						if (TSharedPtr<TOnlineAsyncOp<OpType>> PinnedThis = WeakThis.Pin())
-						{
-							Steps[StepToExecute]->Execute();
-						}
-					});
-			}
+			ExecuteNextStepImpl(OpType::Name, this->AsShared());
 		}
 	}
 
@@ -1502,62 +1570,14 @@ public:
 	{
 		if (!IsComplete())
 		{
-			const int32 StepToExecute = NextStep == 0 ? 0 : NextStep - 1;
-			if (StepToExecute < Steps.Num())
-			{
-				Execute(Steps[StepToExecute]->GetExecutionPolicy(),
-					[this, StepToExecute, WeakThis = TWeakPtr<TOnlineAsyncOp<OpType>>(this->AsShared())]()
-					{
-						if (TSharedPtr<TOnlineAsyncOp<OpType>> PinnedThis = WeakThis.Pin())
-						{
-							Steps[StepToExecute]->Execute();
-						}
-					});
-			}
+			ExecuteRepeatStepImpl(OpType::Name, this->AsShared());
 		}
 	}
 
-	void AddStep(TUniquePtr<Private::IStep>&& Step)
-	{
-		Steps.Add(MoveTemp(Step));
-	}
-
 	template <typename CallableType>
 	void Execute(FOnlineAsyncExecutionPolicy ExecutionPolicy, CallableType&& Callable)
 	{
-		switch (ExecutionPolicy.GetExecutionPolicy())
-		{
-		case EOnlineAsyncExecutionPolicy::RunOnGameThread:
-			if (IsInGameThread())
-			{
-				Callable();
-			}
-			else
-			{
-				ExecuteOnGameThread(OpType::Name, MoveTemp(Callable));
-			}
-			break;
-
-		case EOnlineAsyncExecutionPolicy::RunOnNextTick:
-			FTSTicker::GetCoreTicker().AddTicker(FTickerDelegate::CreateLambda([LambdaCallable = Callable](float dt) mutable
-				{
-					ExecuteOnGameThread(OpType::Name, MoveTemp(LambdaCallable));
-					return false;
-				}));
-			break;
-
-		case EOnlineAsyncExecutionPolicy::RunOnThreadPool:
-			Async(EAsyncExecution::ThreadPool, MoveTemp(Callable));
-			break;
-
-		case EOnlineAsyncExecutionPolicy::RunOnTaskGraph:
-			Async(EAsyncExecution::TaskGraph, MoveTemp(Callable));
-			break;
-
-		case EOnlineAsyncExecutionPolicy::RunImmediately:
-			Callable();
-			break;
-		}
+		ExecuteImpl(OpType::Name, ExecutionPolicy, Callable);
 	}
 
 	TOnlineEvent<void(const TOnlineAsyncOp<OpType>&)> OnStart() { return OnStartEvent; }
@@ -1806,10 +1826,8 @@ PRAGMA_ENABLE_DEPRECATION_WARNINGS
 
 	TSharedRef<FAsyncOpSharedState> SharedState;
 	TArray<TSharedRef<FAsyncOpSharedHandleState>> SharedHandleStates;
-	TArray<TUniquePtr<Private::IStep>> Steps;
 	TOnlineEventCallable<void(const TOnlineAsyncOp<OpType>&)> OnStartEvent;
 	TOnlineEventCallable<void(const TOnlineAsyncOp<OpType>&, const TOnlineResult<OpType>&)> OnCompleteEvent;
-	int32 NextStep = 0;
 	double OpStartTimeAbsoluteSeconds;
 
 	friend class FOnlineAsyncOpCache;
