[AutoRTFM] Use a transactionally-safe lock when looking up online IDs.

diff --git a/Engine/Plugins/Online/OnlineServices/Source/OnlineServicesCommon/Public/Online/OnlineIdCommon.h b/Engine/Plugins/Online/OnlineServices/Source/OnlineServicesCommon/Public/Online/OnlineIdCommon.h
--- a/Engine/Plugins/Online/OnlineServices/Source/OnlineServicesCommon/Public/Online/OnlineIdCommon.h
+++ b/Engine/Plugins/Online/OnlineServices/Source/OnlineServicesCommon/Public/Online/OnlineIdCommon.h
@@ -3,6 +3,7 @@
 #pragma once
 
 #include "Misc/ScopeRWLock.h"
+#include "Misc/TransactionallySafeRWLock.h"
 #include "Online/CoreOnline.h"
 
 namespace UE::Online {
@@ -33,7 +34,7 @@ public:
 
 	HandleType FindHandle(const IdValueType& IdValue) const
 	{
-		FReadScopeLock ReadLock(Lock);
+		UE::TReadScopeLock ReadLock(Lock);
 		if (const HandleType* FoundHandle = IdValueToHandleMap.Find(IdValue))
 		{
 			return *FoundHandle;
@@ -47,7 +48,7 @@ public:
 		if (!Handle.IsValid())
 		{
 			// Take a write lock, check again if we already have a handle, or insert a new element.
-			FWriteScopeLock WriteLock(Lock);
+			UE::TWriteScopeLock WriteLock(Lock);
 			if (const HandleType* FoundHandle = IdValueToHandleMap.Find(IdValue))
 			{
 				Handle = *FoundHandle;
@@ -68,7 +69,7 @@ public:
 	{
 		if (ValidateOnlineId(Handle))
 		{
-			FReadScopeLock ReadLock(Lock);
+			UE::TReadScopeLock ReadLock(Lock);
 			if (IdValues.IsValidIndex(Handle.GetHandle() - 1))
 			{
 				return IdValues[Handle.GetHandle() -1];
@@ -85,7 +86,7 @@ public:
 private:
 	EOnlineServices OnlineServicesType;
 
-	mutable FRWLock Lock;
+	mutable FTransactionallySafeRWLock Lock;
 
 	TArray<IdValueType> IdValues;
 	TMap<IdValueType, HandleType> IdValueToHandleMap;
