- Add logging when registering/unregistering factories

#rnx

diff --git a/Engine/Plugins/Online/OnlineServices/Source/OnlineServicesInterface/Private/Online/OnlineServicesRegistry.cpp b/Engine/Plugins/Online/OnlineServices/Source/OnlineServicesInterface/Private/Online/OnlineServicesRegistry.cpp
--- a/Engine/Plugins/Online/OnlineServices/Source/OnlineServicesInterface/Private/Online/OnlineServicesRegistry.cpp
+++ b/Engine/Plugins/Online/OnlineServices/Source/OnlineServicesInterface/Private/Online/OnlineServicesRegistry.cpp
@@ -58,6 +58,8 @@ void FOnlineServicesRegistry::RegisterServicesFactory(EOnlineServices OnlineServ
 	FFactoryAndPriority* ExistingFactoryAndPriority = ServicesFactories.Find(OnlineServices);
 	if (ExistingFactoryAndPriority == nullptr || ExistingFactoryAndPriority->Priority < Priority)
 	{
+		UE_LOG(LogOnlineServices, Log, TEXT("%hs Type=%s Priority=%d"),
+			__FUNCTION__, LexToString(OnlineServices), Priority);
 		ServicesFactories.Add(OnlineServices, FFactoryAndPriority(MoveTemp(Factory), Priority));
 	}
 }
@@ -69,6 +71,8 @@ void FOnlineServicesRegistry::UnregisterServicesFactory(EOnlineServices OnlineSe
 	FFactoryAndPriority* ExistingFactoryAndPriority = ServicesFactories.Find(OnlineServices);
 	if (ExistingFactoryAndPriority != nullptr && ExistingFactoryAndPriority->Priority == Priority)
 	{
+		UE_LOG(LogOnlineServices, Log, TEXT("%hs Type=%s Priority=%d"),
+			__FUNCTION__, LexToString(OnlineServices), Priority);
 		ServicesFactories.Remove(OnlineServices);
 	}
 
