Use serial queues on top of parallel queue
Use GetJoinableOp/GetMergeableOp for operations that need it
Named local event structures the same way
Replace remaining uses of UTF8_TO_TCHAR by FString construction from const UTF8CHAR*

#Synced-CL 48981444

diff --git a/Engine/Plugins/Online/OnlineServices/Source/OnlineServicesCommon/Public/Online/OnlineAsyncOpQueue.h b/Engine/Plugins/Online/OnlineServices/Source/OnlineServicesCommon/Public/Online/OnlineAsyncOpQueue.h
--- a/Engine/Plugins/Online/OnlineServices/Source/OnlineServicesCommon/Public/Online/OnlineAsyncOpQueue.h
+++ b/Engine/Plugins/Online/OnlineServices/Source/OnlineServicesCommon/Public/Online/OnlineAsyncOpQueue.h
@@ -34,6 +34,8 @@ public:
 		TryStartOperations();
 	}
 
+	virtual bool IsIdle() const = 0;
+
 	virtual void TryStartOperations() = 0;
 
 protected:
@@ -136,6 +138,11 @@ public:
 	{
 	}
 
+	virtual bool IsIdle() const override 
+	{ 
+		return QueuedOperations.IsEmpty() && InFlightOperations.IsEmpty();
+	}
+
 	using FOnlineAsyncOpQueue::Enqueue;
 
 	void Enqueue(TSharedRef<IWrappedOp>& Operation)
@@ -198,6 +205,11 @@ public:
 		}
 	}
 
+	virtual bool IsIdle() const override
+	{
+		return QueuedOperations.IsEmpty() && !InFlightOperation.IsValid();
+	}
+
 protected:
 	FOnlineAsyncOpQueueParallel& ParentQueue;
 
