Fixed System.Engine.Online.ClientLobbyDataTest failure due to different ordering of elements with compact set enabled.

RemoveStable has a higher perf cost, so work is logged to harden the test and potentially revert this change.

#Synced-CL 49850619

diff --git a/Engine/Plugins/Online/OnlineServices/Source/OnlineServicesInterface/Private/Online/Schema.cpp b/Engine/Plugins/Online/OnlineServices/Source/OnlineServicesInterface/Private/Online/Schema.cpp
--- a/Engine/Plugins/Online/OnlineServices/Source/OnlineServicesInterface/Private/Online/Schema.cpp
+++ b/Engine/Plugins/Online/OnlineServices/Source/OnlineServicesInterface/Private/Online/Schema.cpp
@@ -753,9 +753,11 @@ bool FSchemaRegistry::ParseConfig(const FSchemaRegistryDescriptorConfig& Config)
 								// Todo: Reuse descriptor for assignment until it is full.
 
 								// Remove found descriptor from buckets.
-								PublicAttributes.Remove(FoundServiceAttributeDescriptor);
-								PrivateAttributes.Remove(FoundServiceAttributeDescriptor);
-								SearchableAttributes.Remove(FoundServiceAttributeDescriptor);
+								// Stable remove to preserve ordering with compact tset enabled.
+								// Automated tests (eg System.Engine.Online.ClientLobbyDataTest) depend on this order.
+								PublicAttributes.RemoveStable(FoundServiceAttributeDescriptor);
+								PrivateAttributes.RemoveStable(FoundServiceAttributeDescriptor);
+								SearchableAttributes.RemoveStable(FoundServiceAttributeDescriptor);
 
 								// Create service attribute definition.
 								FSchemaServiceAttributeDefinition& SchemaServiceAttributeDefinition = SchemaCategoryDefinition.ServiceAttributeDefinitions.Add(FoundServiceAttributeDescriptor->Id);
