Fix for assertion on PIE teardown caused by obsolete tickers

#Synced-CL 51330672

diff --git a/Engine/Plugins/Online/OnlineServices/Source/OnlineServicesCommonEngineUtils/Private/Online/OnlineServicesCommonEngineUtils.cpp b/Engine/Plugins/Online/OnlineServices/Source/OnlineServicesCommonEngineUtils/Private/Online/OnlineServicesCommonEngineUtils.cpp
--- a/Engine/Plugins/Online/OnlineServices/Source/OnlineServicesCommonEngineUtils/Private/Online/OnlineServicesCommonEngineUtils.cpp
+++ b/Engine/Plugins/Online/OnlineServices/Source/OnlineServicesCommonEngineUtils/Private/Online/OnlineServicesCommonEngineUtils.cpp
@@ -17,9 +17,12 @@ UWorld* GetWorldForOnline(FName InstanceName)
 #if WITH_EDITOR
 	if (InstanceName.ToString() != LexToString(UE::Online::EOnlineServices::Default) && InstanceName != NAME_None)
 	{
-		FWorldContext& WorldContext = GEngine->GetWorldContextFromHandleChecked(InstanceName);
-		check(WorldContext.WorldType == EWorldType::Game || WorldContext.WorldType == EWorldType::PIE);
-		World = WorldContext.World();
+		// The world context is not guaranteed to always be valid (e.g: PIE teardown)
+		if (const FWorldContext* WorldContext = GEngine->GetWorldContextFromHandle(InstanceName))
+		{
+			check(WorldContext->WorldType == EWorldType::Game || WorldContext->WorldType == EWorldType::PIE);
+			World = WorldContext->World();
+		}
 	}
 	else
 #endif
diff --git a/Engine/Plugins/Online/OnlineSubsystem/Source/Private/OnlineSubsystemImpl.cpp b/Engine/Plugins/Online/OnlineSubsystem/Source/Private/OnlineSubsystemImpl.cpp
--- a/Engine/Plugins/Online/OnlineSubsystem/Source/Private/OnlineSubsystemImpl.cpp
+++ b/Engine/Plugins/Online/OnlineSubsystem/Source/Private/OnlineSubsystemImpl.cpp
@@ -141,6 +141,12 @@ bool FOnlineSubsystemImpl::Tick(float DeltaTime)
 	return true;
 }
 
+void FOnlineSubsystemImpl::PIEEndCleanup()
+{
+	// As part of the Editor teardown process, we stop any tickers
+	StopTicker();
+}
+
 void FOnlineSubsystemImpl::InitNamedInterfaces()
 {
 	NamedInterfaces = NewObject<UNamedInterfaces>();
diff --git a/Engine/Plugins/Online/OnlineSubsystem/Source/Public/OnlineSubsystemImpl.h b/Engine/Plugins/Online/OnlineSubsystem/Source/Public/OnlineSubsystemImpl.h
--- a/Engine/Plugins/Online/OnlineSubsystem/Source/Public/OnlineSubsystemImpl.h
+++ b/Engine/Plugins/Online/OnlineSubsystem/Source/Public/OnlineSubsystemImpl.h
@@ -139,6 +139,9 @@ public:
 	// FTSTickerObjectBase
 	UE_API virtual bool Tick(float DeltaTime) override;
 
+	/* Performs cleanup tasks for editor teardown */
+	UE_API void PIEEndCleanup();
+
 	/**
 	 * Modify a response string so that it can be logged cleanly
 	 *
diff --git a/Engine/Plugins/Online/OnlineSubsystemUtils/Source/OnlineSubsystemUtils/Private/OnlineEngineInterfaceImpl.cpp b/Engine/Plugins/Online/OnlineSubsystemUtils/Source/OnlineSubsystemUtils/Private/OnlineEngineInterfaceImpl.cpp
--- a/Engine/Plugins/Online/OnlineSubsystemUtils/Source/OnlineSubsystemUtils/Private/OnlineEngineInterfaceImpl.cpp
+++ b/Engine/Plugins/Online/OnlineSubsystemUtils/Source/OnlineSubsystemUtils/Private/OnlineEngineInterfaceImpl.cpp
@@ -8,6 +8,7 @@
 #include "Interfaces/OnlineSessionDelegates.h"
 #include "Online.h"
 #include "OnlineSessionSettings.h"
+#include "OnlineSubsystemImpl.h"
 
 #include UE_INLINE_GENERATED_CPP_BY_NAME(OnlineEngineInterfaceImpl)
 
@@ -734,6 +735,16 @@ void UOnlineEngineInterfaceImpl::OnPIELoginComplete(int32 LocalUserNum, bool bWa
 	InCompletionDelegate.ExecuteIfBound(LocalUserNum, bWasSuccessful, Error);
 }
 
+void UOnlineEngineInterfaceImpl::PIEEndCleanup(FName OnlineIdentifier)
+{
+	IOnlineSubsystem* Subsystem = IOnlineSubsystem::Get(OnlineIdentifier);
+
+	if (FOnlineSubsystemImpl* SubsystemImpl = StaticCast<FOnlineSubsystemImpl*>(Subsystem))
+	{
+		SubsystemImpl->PIEEndCleanup();
+	}
+}
+
 #endif
 
 void UOnlineEngineInterfaceImpl::InitCompatibilityInterface()
diff --git a/Engine/Plugins/Online/OnlineSubsystemUtils/Source/OnlineSubsystemUtils/Private/OnlineEngineInterfaceImpl.h b/Engine/Plugins/Online/OnlineSubsystemUtils/Source/OnlineSubsystemUtils/Private/OnlineEngineInterfaceImpl.h
--- a/Engine/Plugins/Online/OnlineSubsystemUtils/Source/OnlineSubsystemUtils/Private/OnlineEngineInterfaceImpl.h
+++ b/Engine/Plugins/Online/OnlineSubsystemUtils/Source/OnlineSubsystemUtils/Private/OnlineEngineInterfaceImpl.h
@@ -163,6 +163,7 @@ public:
 	UE_API virtual FString GetPIELoginCommandLineArgs(int32 Index) override;
 	UE_API virtual void SetForceDedicated(FName OnlineIdentifier, bool bForce) override;
 	UE_API virtual void LoginPIEInstance(FName OnlineIdentifier, int32 LocalUserNum, int32 PIELoginNum, FOnPIELoginComplete& CompletionDelegate) override;
+	UE_API virtual void PIEEndCleanup(FName OnlineIdentifier) override;
 #endif
 
 private:
diff --git a/Engine/Source/Editor/UnrealEd/Private/PlayLevel.cpp b/Engine/Source/Editor/UnrealEd/Private/PlayLevel.cpp
--- a/Engine/Source/Editor/UnrealEd/Private/PlayLevel.cpp
+++ b/Engine/Source/Editor/UnrealEd/Private/PlayLevel.cpp
@@ -367,6 +367,12 @@ void UEditorEngine::EndPlayMap()
 		}
 	}
 
+	// We perform subsystem cleanup before removing the world contexts
+	for (const FName& OnlineIdentifier : OnlineIdentifiers)
+	{
+		UOnlineEngineInterface::Get()->PIEEndCleanup(OnlineIdentifier);	
+	}
+
 	// Only remove from list after all PIE worlds have been torn down, to prevent callbacks from trying to access removed worldcontexts
 	for (int32 WorldIdx = WorldList.Num() - 1; WorldIdx >= 0; --WorldIdx)
 	{
diff --git a/Engine/Source/Runtime/Engine/Public/Net/OnlineEngineInterface.h b/Engine/Source/Runtime/Engine/Public/Net/OnlineEngineInterface.h
--- a/Engine/Source/Runtime/Engine/Public/Net/OnlineEngineInterface.h
+++ b/Engine/Source/Runtime/Engine/Public/Net/OnlineEngineInterface.h
@@ -256,6 +256,8 @@ public:
 	virtual void SetForceDedicated(FName OnlineIdentifier, bool bForce) {}
 	/** Login a local user to the online subsystem before starting a Play In Editor (PIE) instance, using credentials stored in the editor */
 	virtual void LoginPIEInstance(FName OnlineIdentifier, int32 LocalUserNum, int32 PIELoginNum, FOnPIELoginComplete& CompletionDelegate) { CompletionDelegate.ExecuteIfBound(LocalUserNum, false, TEXT("Not implemented")); }
+	/** Perform any relevant cleanup operations for the online subsystem during Play In Editor (PIE) end/teardown */
+	virtual void PIEEndCleanup(FName OnlineIdentifier) {}
 
 private:
 
