Fix missing RichPresenceString in FMutations merge
Also refactored PresenceCommon::PartialUpdatePresence to use ApplyPresenceMutations to prevent future field omissions.

#Synced-CL 52467806

diff --git a/Engine/Plugins/Online/OnlineServices/Source/OnlineServicesCommon/Private/Online/PresenceCommon.cpp b/Engine/Plugins/Online/OnlineServices/Source/OnlineServicesCommon/Private/Online/PresenceCommon.cpp
--- a/Engine/Plugins/Online/OnlineServices/Source/OnlineServicesCommon/Private/Online/PresenceCommon.cpp
+++ b/Engine/Plugins/Online/OnlineServices/Source/OnlineServicesCommon/Private/Online/PresenceCommon.cpp
@@ -143,34 +143,8 @@ TOnlineAsyncOpHandle<FPartialUpdatePresence> FPresenceCommon::PartialUpdatePrese
 	TOnlineResult<FGetCachedPresence> CachedPresence = GetCachedPresence({Params.LocalAccountId, Params.LocalAccountId});
 	if (CachedPresence.IsOk())
 	{
-		TSharedRef<FUserPresence> NewPresence = MakeShared<FUserPresence>();
 		TSharedRef<const FUserPresence> ExistingPresence = CachedPresence.GetOkValue().Presence;
-		*NewPresence = *ExistingPresence;
-
-		if (Params.Mutations.Status.IsSet())
-		{
-			NewPresence->Status = Params.Mutations.Status.GetValue();
-		}
-
-		if (Params.Mutations.GameStatus.IsSet())
-		{
-			NewPresence->GameStatus = Params.Mutations.GameStatus.GetValue();
-		}
-
-		if (Params.Mutations.Joinability.IsSet())
-		{
-			NewPresence->Joinability = Params.Mutations.Joinability.GetValue();
-		}
-
-		for (const TPair<FString, FPresenceProperty>& UpdatedProperty : Params.Mutations.UpdatedProperties)
-		{
-			NewPresence->Properties.Add(UpdatedProperty.Key, UpdatedProperty.Value);
-		}
-
-		for (const FString& RemovedProperty : Params.Mutations.RemovedProperties)
-		{
-			NewPresence->Properties.Remove(RemovedProperty);
-		}
+		TSharedRef<FUserPresence> NewPresence = ApplyPresenceMutations(*ExistingPresence, Params.Mutations);
 
 		TOnlineAsyncOpRef<FPartialUpdatePresence> Op = GetOp<FPartialUpdatePresence>(MoveTemp(Params));
 		TOnlineAsyncOpHandle<FUpdatePresence> UpdatePresenceResult = UpdatePresence({Op->GetParams().LocalAccountId, NewPresence});
diff --git a/Engine/Plugins/Online/OnlineServices/Source/OnlineServicesInterface/Public/Online/Presence.h b/Engine/Plugins/Online/OnlineServices/Source/OnlineServicesInterface/Public/Online/Presence.h
--- a/Engine/Plugins/Online/OnlineServices/Source/OnlineServicesInterface/Public/Online/Presence.h
+++ b/Engine/Plugins/Online/OnlineServices/Source/OnlineServicesInterface/Public/Online/Presence.h
@@ -249,6 +249,11 @@ inline FPartialUpdatePresence::Params::FMutations& FPartialUpdatePresence::Param
 		StatusString = MoveTemp(NewParams.StatusString);
 	}
 
+	if (NewParams.RichPresenceString.IsSet())
+	{
+		RichPresenceString = MoveTemp(NewParams.RichPresenceString);
+	}
+
 	// Insert any updated keys / remove from removed
 	for (TPair<FString, FPresenceProperty>& NewUpdatedProperty : NewParams.UpdatedProperties)
 	{
